DMX - Add GDTF export functinality, unit testing, improve GDTF data structure

diff --git a/Engine/Plugins/VirtualProduction/DMX/DMXEngine/Config/BaseDMXEngine.ini b/Engine/Plugins/VirtualProduction/DMX/DMXEngine/Config/BaseDMXEngine.ini
--- a/Engine/Plugins/VirtualProduction/DMX/DMXEngine/Config/BaseDMXEngine.ini
+++ b/Engine/Plugins/VirtualProduction/DMX/DMXEngine/Config/BaseDMXEngine.ini
@@ -4,3 +4,4 @@
 
 ; UE5.4 to UE5.5
 +FunctionRedirects=(OldName="GetAllDMXLibraries",NewName="/Script/DMXEngine.DMXSubsystem.LoadDMXLibrariesSynchronous")
++StructRedirects=(OldName="DMXColorCIE1931xyY", NewName="/Script/DMXGDTF.DMXGDTFColorCIE1931xyY")
diff --git a/Engine/Plugins/VirtualProduction/DMX/DMXEngine/Source/DMXEditor/Private/Factories/DMXGDTFFactory.cpp b/Engine/Plugins/VirtualProduction/DMX/DMXEngine/Source/DMXEditor/Private/Factories/DMXGDTFFactory.cpp
--- a/Engine/Plugins/VirtualProduction/DMX/DMXEngine/Source/DMXEditor/Private/Factories/DMXGDTFFactory.cpp
+++ b/Engine/Plugins/VirtualProduction/DMX/DMXEngine/Source/DMXEditor/Private/Factories/DMXGDTFFactory.cpp
@@ -99,7 +99,9 @@ UObject* UDMXGDTFFactory::FactoryCreateFile(UClass* InClass, UObject* InParent,
     // Create import options
 	using namespace UE::DMX;
     FDMXGDTFImportArgs ImportArgs;
-    ImportArgs.Name = InName;
+
+	const FString BaseFilename = FPaths::GetBaseFilename(InName.ToString());
+    ImportArgs.Name = *BaseFilename;
     ImportArgs.Parent = InParent;
     ImportArgs.Filename = InFilename;
     ImportArgs.Flags = Flags;
diff --git a/Engine/Plugins/VirtualProduction/DMX/DMXEngine/Source/DMXEditor/Private/Factories/DMXGDTFImporter.cpp b/Engine/Plugins/VirtualProduction/DMX/DMXEngine/Source/DMXEditor/Private/Factories/DMXGDTFImporter.cpp
--- a/Engine/Plugins/VirtualProduction/DMX/DMXEngine/Source/DMXEditor/Private/Factories/DMXGDTFImporter.cpp
+++ b/Engine/Plugins/VirtualProduction/DMX/DMXEngine/Source/DMXEditor/Private/Factories/DMXGDTFImporter.cpp
@@ -61,7 +61,8 @@ namespace UE::DMX
 		}
 
 		// Create GDTF
-		UDMXImportGDTF* NewGDTF = NewObject<UDMXImportGDTF>(ImportArgs.Parent.Get(), ImportArgs.Name, ImportArgs.Flags | RF_Public);
+		const FName UniqueName = MakeUniqueObjectName(ImportArgs.Parent.Get(), UDMXImportGDTF::StaticClass(), ImportArgs.Name);
+		UDMXImportGDTF* NewGDTF = NewObject<UDMXImportGDTF>(ImportArgs.Parent.Get(), UniqueName, ImportArgs.Flags | RF_Public);
 
 		// Set Asset Import Data
 		UDMXGDTFAssetImportData* GDTFAssetImportData = NewGDTF->GetGDTFAssetImportData();
diff --git a/Engine/Plugins/VirtualProduction/DMX/DMXEngine/Source/DMXRuntime/Private/MVR/Types/DMXMVRFixtureNode.cpp b/Engine/Plugins/VirtualProduction/DMX/DMXEngine/Source/DMXRuntime/Private/MVR/Types/DMXMVRFixtureNode.cpp
--- a/Engine/Plugins/VirtualProduction/DMX/DMXEngine/Source/DMXRuntime/Private/MVR/Types/DMXMVRFixtureNode.cpp
+++ b/Engine/Plugins/VirtualProduction/DMX/DMXEngine/Source/DMXRuntime/Private/MVR/Types/DMXMVRFixtureNode.cpp
@@ -255,7 +255,7 @@ namespace UE::DMXRuntime::DMXMVRFixture::Private
 				return;
 			}
 
-			FDMXColorCIE1931xyY Result;
+			FDMXGDTFColorCIE1931xyY Result;
 			bool bSuccess = LexTryParseString<float>(Result.X, *CIEColorArray[0]);
 			bSuccess = bSuccess && LexTryParseString<float>(Result.Y, *CIEColorArray[1]);
 			bSuccess = bSuccess && LexTryParseString<float>(Result.YY, *CIEColorArray[2]);
diff --git a/Engine/Plugins/VirtualProduction/DMX/DMXEngine/Source/DMXRuntime/Public/DMXOptionalTypes.h b/Engine/Plugins/VirtualProduction/DMX/DMXEngine/Source/DMXRuntime/Public/DMXOptionalTypes.h
--- a/Engine/Plugins/VirtualProduction/DMX/DMXEngine/Source/DMXRuntime/Public/DMXOptionalTypes.h
+++ b/Engine/Plugins/VirtualProduction/DMX/DMXEngine/Source/DMXRuntime/Public/DMXOptionalTypes.h
@@ -437,7 +437,7 @@ struct DMXRUNTIME_API FDMXOptionalColorCIE1931
 	GENERATED_BODY()
 
 	/** The Optional Value */
-	TOptional<FDMXColorCIE1931xyY> Value;
+	TOptional<FDMXGDTFColorCIE1931xyY> Value;
 
 	/** Resets the optional */
 	void Reset() { Value.Reset(); }
@@ -446,9 +446,9 @@ struct DMXRUNTIME_API FDMXOptionalColorCIE1931
 	bool IsSet() const { return Value.IsSet(); }
 	
 	/** Returns the Value of the optional, only call when the value is set. */
-	const FDMXColorCIE1931xyY& GetValue() const { return Value.GetValue(); }
+	const FDMXGDTFColorCIE1931xyY& GetValue() const { return Value.GetValue(); }
 
-	const FDMXOptionalColorCIE1931& operator=(const FDMXColorCIE1931xyY& Transform)
+	const FDMXOptionalColorCIE1931& operator=(const FDMXGDTFColorCIE1931xyY& Transform)
 	{
 		Value = Transform;
 		return *this;
diff --git a/Engine/Plugins/VirtualProduction/DMX/DMXEngine/Source/DMXRuntime/Public/DMXTypes.h b/Engine/Plugins/VirtualProduction/DMX/DMXEngine/Source/DMXRuntime/Public/DMXTypes.h
--- a/Engine/Plugins/VirtualProduction/DMX/DMXEngine/Source/DMXRuntime/Public/DMXTypes.h
+++ b/Engine/Plugins/VirtualProduction/DMX/DMXEngine/Source/DMXRuntime/Public/DMXTypes.h
@@ -3,15 +3,15 @@
 #pragma once
 
 #include "DMXAttribute.h"
-
-#include "CoreMinimal.h"
+#include "DMXGDTFColorCIE1931xyY.h"
 #include "Templates/SubclassOf.h"
 
 #include "DMXTypes.generated.h"
 
 class UDMXLibrary;
 
-
+UE_DEPRECATED(5.5, "FDMXColorCIE1931xyY is deprecated. Please use FDMXGDTFColorCIE1931xyY instead.")
+typedef FDMXGDTFColorCIE1931xyY FDMXColorCIE1931xyY;
 
 /** Holds an array Attribute Names with their normalized Values (expand the property to see the map) */
 USTRUCT(BlueprintType, Category = "DMX")
@@ -23,49 +23,6 @@ struct DMXRUNTIME_API FDMXNormalizedAttributeValueMap
 	TMap<FDMXAttributeName, float> Map;
 };
 
-
-/** xyY color representation in the CIE 1931 color space, as typically used in GDTF and MVR */
-USTRUCT(BlueprintType, Category = "DMX")
-struct DMXRUNTIME_API FDMXColorCIE1931xyY
-{
-	GENERATED_BODY()
-
-	FString ToString() const
-	{
-		return FString::FromInt(X) + TEXT(", ") + FString::FromInt(Y) + TEXT(", ") + FString::FromInt(YY);
-	}
-
-	/** x */
-	UPROPERTY(BlueprintReadWrite, EditAnywhere, Category = "DMX")
-	float X = 0.f;
-
-	/** y */
-	UPROPERTY(BlueprintReadWrite, EditAnywhere, Category = "DMX")
-	float Y = 0.f;
-
-	/** Y */
-	UPROPERTY(BlueprintReadWrite, EditAnywhere, Category = "DMX")
-	float YY = 0.f;
-
-	friend bool operator==(const FDMXColorCIE1931xyY& A, const FDMXColorCIE1931xyY& B)
-	{
-		return A.X == B.X && A.Y == B.Y && A.YY == B.YY;
-	}
-
-	friend bool operator!=(const FDMXColorCIE1931xyY& A, const FDMXColorCIE1931xyY& B)
-	{
-		return !(A == B);
-	}
-
-	friend FArchive& operator<<(FArchive& Ar, FDMXColorCIE1931xyY& ColorCIE1931)
-	{
-		Ar << ColorCIE1931.X;
-		Ar << ColorCIE1931.Y;
-		Ar << ColorCIE1931.YY;
-		return Ar;
-	}
-};
-
 USTRUCT()
 struct FDMXByteArray64
 {
diff --git a/Engine/Plugins/VirtualProduction/DMX/DMXGDTF/DMXGDTF.uplugin b/Engine/Plugins/VirtualProduction/DMX/DMXGDTF/DMXGDTF.uplugin
--- a/Engine/Plugins/VirtualProduction/DMX/DMXGDTF/DMXGDTF.uplugin
+++ b/Engine/Plugins/VirtualProduction/DMX/DMXGDTF/DMXGDTF.uplugin
@@ -3,7 +3,7 @@
 	"Version": 1,
 	"VersionName": "1.0",
 	"FriendlyName": "DMX GDTF",
-	"Description": "Implementation of GDTF standard using Unreal Engine types",
+	"Description": "Implementation of the GDTF standard using Unreal Engine types",
 	"Category": "Virtual Production",
 	"CreatedBy": "Epic Games, Inc.",
 	"CreatedByURL": "https://epicgames.com",
@@ -20,6 +20,11 @@
 			"Type": "Runtime",
 			"LoadingPhase": "Default"
 		},
+		{
+			"Name": "DMXGDTFTests",
+			"Type": "Editor",
+			"LoadingPhase": "Default"
+		},
 		{
 			"Name": "DMXZip",
 			"Type": "Runtime",
diff --git a/Engine/Plugins/VirtualProduction/DMX/DMXGDTF/Source/DMXGDTF/Internal/GDTF/AttributeDefinitions/DMXGDTFActivationGroup.h b/Engine/Plugins/VirtualProduction/DMX/DMXGDTF/Source/DMXGDTF/Internal/GDTF/AttributeDefinitions/DMXGDTFActivationGroup.h
--- a/Engine/Plugins/VirtualProduction/DMX/DMXGDTF/Source/DMXGDTF/Internal/GDTF/AttributeDefinitions/DMXGDTFActivationGroup.h
+++ b/Engine/Plugins/VirtualProduction/DMX/DMXGDTF/Source/DMXGDTF/Internal/GDTF/AttributeDefinitions/DMXGDTFActivationGroup.h
@@ -18,6 +18,7 @@ namespace UE::DMX::GDTF
 		//~ Begin FDMXGDTFNode interface
 		virtual const TCHAR* GetXmlTag() const override { return TEXT("ActivationGroup"); }
 		virtual void Initialize(const FXmlNode& XmlNode) override;
+		virtual FXmlNode* CreateXmlNode(FXmlNode& Parent) override;
 		//~ End FDMXGDTFNode interface
 
 		/** The unique name of the activation group. */
diff --git a/Engine/Plugins/VirtualProduction/DMX/DMXGDTF/Source/DMXGDTF/Internal/GDTF/AttributeDefinitions/DMXGDTFAttribute.h b/Engine/Plugins/VirtualProduction/DMX/DMXGDTF/Source/DMXGDTF/Internal/GDTF/AttributeDefinitions/DMXGDTFAttribute.h
--- a/Engine/Plugins/VirtualProduction/DMX/DMXGDTF/Source/DMXGDTF/Internal/GDTF/AttributeDefinitions/DMXGDTFAttribute.h
+++ b/Engine/Plugins/VirtualProduction/DMX/DMXGDTF/Source/DMXGDTF/Internal/GDTF/AttributeDefinitions/DMXGDTFAttribute.h
@@ -2,6 +2,7 @@
 
 #pragma once
 
+#include "DMXGDTFColorCIE1931xyY.h"
 #include "GDTF/DMXGDTFNode.h"
 #include "Math/Vector.h"
 
@@ -24,6 +25,7 @@ namespace UE::DMX::GDTF
 		//~ Begin FDMXGDTFNode interface
 		virtual const TCHAR* GetXmlTag() const override { return TEXT("Attribute"); }
 		virtual void Initialize(const FXmlNode& XmlNode) override;
+		virtual FXmlNode* CreateXmlNode(FXmlNode& Parent) override;
 		//~ End FDMXGDTFNode interface
 
 		/** The unique name of the attribute. */
@@ -45,7 +47,7 @@ namespace UE::DMX::GDTF
 		EDMXGDTFPhysicalUnit PhysicalUnit;
 
 		/** (Optional) Defines the color for the attribute. */
-		FVector Color;
+		FDMXGDTFColorCIE1931xyY Color;
 
 		/** As children the attribute node has a list of a subphysical units. */
 		TArray<TSharedPtr<FDMXGDTFSubphysicalUnit>> SubpyhsicalUnitArray;
diff --git a/Engine/Plugins/VirtualProduction/DMX/DMXGDTF/Source/DMXGDTF/Internal/GDTF/AttributeDefinitions/DMXGDTFAttributeDefinitions.h b/Engine/Plugins/VirtualProduction/DMX/DMXGDTF/Source/DMXGDTF/Internal/GDTF/AttributeDefinitions/DMXGDTFAttributeDefinitions.h
--- a/Engine/Plugins/VirtualProduction/DMX/DMXGDTF/Source/DMXGDTF/Internal/GDTF/AttributeDefinitions/DMXGDTFAttributeDefinitions.h
+++ b/Engine/Plugins/VirtualProduction/DMX/DMXGDTF/Source/DMXGDTF/Internal/GDTF/AttributeDefinitions/DMXGDTFAttributeDefinitions.h
@@ -12,7 +12,7 @@ namespace UE::DMX::GDTF
 	class FDMXGDTFFeatureGroup;
 	class FDMXGDTFFixtureType;
 
-	/** This section defines the attribute DefinitionCollect for the Fixture Type Attributes. */
+	/** This section defines the attribute definition collect for the Fixture Type Attributes. */
 	class DMXGDTF_API FDMXGDTFAttributeDefinitions
 		: public FDMXGDTFNode
 	{
@@ -22,6 +22,7 @@ namespace UE::DMX::GDTF
 		//~ Begin FDMXGDTFNode interface
 		virtual const TCHAR* GetXmlTag() const override { return TEXT("AttributeDefinitions"); }
 		virtual void Initialize(const FXmlNode& InXmlNode) override;
+		virtual FXmlNode* CreateXmlNode(FXmlNode& Parent) override;
 		//~ End FDMXGDTFNode interface
 
 		/** 
diff --git a/Engine/Plugins/VirtualProduction/DMX/DMXGDTF/Source/DMXGDTF/Internal/GDTF/AttributeDefinitions/DMXGDTFFeature.h b/Engine/Plugins/VirtualProduction/DMX/DMXGDTF/Source/DMXGDTF/Internal/GDTF/AttributeDefinitions/DMXGDTFFeature.h
--- a/Engine/Plugins/VirtualProduction/DMX/DMXGDTF/Source/DMXGDTF/Internal/GDTF/AttributeDefinitions/DMXGDTFFeature.h
+++ b/Engine/Plugins/VirtualProduction/DMX/DMXGDTF/Source/DMXGDTF/Internal/GDTF/AttributeDefinitions/DMXGDTFFeature.h
@@ -13,11 +13,12 @@ namespace UE::DMX::GDTF
 		: public FDMXGDTFNode
 	{
 	public:
-		FDMXGDTFFeature(const TWeakPtr<FDMXGDTFFeatureGroup>& InFeatureGroup);
+		FDMXGDTFFeature(const TSharedRef<FDMXGDTFFeatureGroup>& InFeatureGroup);
 
 		//~ Begin FDMXGDTFNode interface
 		virtual const TCHAR* GetXmlTag() const override { return TEXT("Feature"); }
 		virtual void Initialize(const FXmlNode& XmlNode) override;
+		virtual FXmlNode* CreateXmlNode(FXmlNode& Parent) override;
 		//~ End FDMXGDTFNode interface
 
 		/** The unique name of the feature. */
diff --git a/Engine/Plugins/VirtualProduction/DMX/DMXGDTF/Source/DMXGDTF/Internal/GDTF/AttributeDefinitions/DMXGDTFFeatureGroup.h b/Engine/Plugins/VirtualProduction/DMX/DMXGDTF/Source/DMXGDTF/Internal/GDTF/AttributeDefinitions/DMXGDTFFeatureGroup.h
--- a/Engine/Plugins/VirtualProduction/DMX/DMXGDTF/Source/DMXGDTF/Internal/GDTF/AttributeDefinitions/DMXGDTFFeatureGroup.h
+++ b/Engine/Plugins/VirtualProduction/DMX/DMXGDTF/Source/DMXGDTF/Internal/GDTF/AttributeDefinitions/DMXGDTFFeatureGroup.h
@@ -19,6 +19,7 @@ namespace UE::DMX::GDTF
 		//~ Begin FDMXGDTFNode interface
 		virtual const TCHAR* GetXmlTag() const override { return TEXT("FeatureGroup"); }
 		virtual void Initialize(const FXmlNode& XmlNode) override;
+		virtual FXmlNode* CreateXmlNode(FXmlNode& Parent) override;
 		//~ End FDMXGDTFNode interface
 
 		/** The unique name of the feature group. */
diff --git a/Engine/Plugins/VirtualProduction/DMX/DMXGDTF/Source/DMXGDTF/Internal/GDTF/AttributeDefinitions/DMXGDTFSubphysicalUnit.h b/Engine/Plugins/VirtualProduction/DMX/DMXGDTF/Source/DMXGDTF/Internal/GDTF/AttributeDefinitions/DMXGDTFSubphysicalUnit.h
--- a/Engine/Plugins/VirtualProduction/DMX/DMXGDTF/Source/DMXGDTF/Internal/GDTF/AttributeDefinitions/DMXGDTFSubphysicalUnit.h
+++ b/Engine/Plugins/VirtualProduction/DMX/DMXGDTF/Source/DMXGDTF/Internal/GDTF/AttributeDefinitions/DMXGDTFSubphysicalUnit.h
@@ -50,6 +50,7 @@ namespace UE::DMX::GDTF
 		//~ Begin FDMXGDTFNode interface
 		virtual const TCHAR* GetXmlTag() const override { return TEXT("SubphysicalUnit"); }
 		virtual void Initialize(const FXmlNode& XmlNode) override;
+		virtual FXmlNode* CreateXmlNode(FXmlNode& Parent) override;
 		//~ End FDMXGDTFNode interface
 
 		/** Type of subpyhsical unit */
@@ -59,10 +60,10 @@ namespace UE::DMX::GDTF
 		EDMXGDTFPhysicalUnit PhysicalUnit;
 
 		/** The default physical from of the subphysical unit; Unit: as defined in PhysicalUnit; Default value: 0 */
-		double PhysicalFrom = 0.0;
+		float PhysicalFrom = 0.0;
 
 		/** The default physical to of the subphysical unit; Unit: as defined in PhysicalUnit; Default value: 1 */
-		double PhysicalTo = 1.0;
+		float PhysicalTo = 1.0;
 
 		/** The outer attribute */
 		const TWeakPtr<FDMXGDTFAttribute> OuterAttribute;
diff --git a/Engine/Plugins/VirtualProduction/DMX/DMXGDTF/Source/DMXGDTF/Internal/GDTF/DMXGDTFDescription.h b/Engine/Plugins/VirtualProduction/DMX/DMXGDTF/Source/DMXGDTF/Internal/GDTF/DMXGDTFDescription.h
--- a/Engine/Plugins/VirtualProduction/DMX/DMXGDTF/Source/DMXGDTF/Internal/GDTF/DMXGDTFDescription.h
+++ b/Engine/Plugins/VirtualProduction/DMX/DMXGDTF/Source/DMXGDTF/Internal/GDTF/DMXGDTFDescription.h
@@ -2,6 +2,7 @@
 
 #pragma once
 
+#include "GDTF/DMXGDTFVersion.h"
 #include "Templates/SharedPointer.h"
 
 class FXmlFile;
@@ -24,6 +25,12 @@ namespace UE::DMX::GDTF
 		/** Initializes this node and its chilren from a GDTF Description.xml */
 		void InitializeFromDescriptionXml(const TSharedRef<FXmlFile>& DescriptionXml);
 
+		/** Initializes this object from a GDTF fixture type. */
+		void InitializeFromFixtureType(const TSharedRef<FDMXGDTFFixtureType>& InFixtureType);
+
+		/** Exports an Xml file from the node tree in this description. Returns the XML file or nullptr if not a valid XML */
+		TSharedPtr<FXmlFile> ExportAsXml() const;
+
 	private:
 		/** The Fixture Type Child Node */
 		TSharedPtr<FDMXGDTFFixtureType> FixtureType;
diff --git a/Engine/Plugins/VirtualProduction/DMX/DMXGDTF/Source/DMXGDTF/Internal/GDTF/DMXGDTFFixtureType.h b/Engine/Plugins/VirtualProduction/DMX/DMXGDTF/Source/DMXGDTF/Internal/GDTF/DMXGDTFFixtureType.h
--- a/Engine/Plugins/VirtualProduction/DMX/DMXGDTF/Source/DMXGDTF/Internal/GDTF/DMXGDTFFixtureType.h
+++ b/Engine/Plugins/VirtualProduction/DMX/DMXGDTF/Source/DMXGDTF/Internal/GDTF/DMXGDTFFixtureType.h
@@ -10,6 +10,7 @@ namespace UE::DMX::GDTF
 	class FDMXGDTFAttributeDefinitions;
 	class FDMXGDTFDescription;
 	class FDMXGDTFDMXMode;
+	class FDMXGDTFFTPreset;
 	class FDMXGDTFGeometryCollect;
 	class FDMXGDTFModel;
 	class FDMXGDTFPhysicalDescriptions;
@@ -17,7 +18,7 @@ namespace UE::DMX::GDTF
 	class FDMXGDTFRevision;
 	class FDMXGDTFWheel;
 
-	/** The FixtureType node is the starting point of the description of the fixture type within the [GDTF's] XML file. */
+	/** The FixtureType node is the starting point of the description of the fixture type within the XML file. */
 	class DMXGDTF_API FDMXGDTFFixtureType
 		: public FDMXGDTFNode
 	{
@@ -25,6 +26,7 @@ namespace UE::DMX::GDTF
 		//~ Begin FDMXGDTFNode interface
 		virtual const TCHAR* GetXmlTag() const override { return TEXT("FixtureType"); }
 		virtual void Initialize(const FXmlNode& InXmlNode) override;
+		virtual FXmlNode* CreateXmlNode(FXmlNode& Parent) override;
 		//~ End FDMXGDTFNode interface
 
 		/** Name of the fixture type. As it is based on Name. [UE: GDTF version specific, unreal only parses the name] */
@@ -101,6 +103,9 @@ namespace UE::DMX::GDTF
 		/** This section defines the history of device type (XML node <Revisions>). */
 		TArray<TSharedPtr<FDMXGDTFRevision>> Revisions;
 
+		/** Is used to transfer user-defined and fixture type specific presets to other show files. */
+		TArray<TSharedPtr<FDMXGDTFFTPreset>> FTPresets;
+
 		/** Is used to specify supported protocols. */
 		TSharedPtr<FDMXGDTFProtocols> Protocols;	
 	};
diff --git a/Engine/Plugins/VirtualProduction/DMX/DMXGDTF/Source/DMXGDTF/Internal/GDTF/DMXGDTFNode.h b/Engine/Plugins/VirtualProduction/DMX/DMXGDTF/Source/DMXGDTF/Internal/GDTF/DMXGDTFNode.h
--- a/Engine/Plugins/VirtualProduction/DMX/DMXGDTF/Source/DMXGDTF/Internal/GDTF/DMXGDTFNode.h
+++ b/Engine/Plugins/VirtualProduction/DMX/DMXGDTF/Source/DMXGDTF/Internal/GDTF/DMXGDTFNode.h
@@ -2,6 +2,7 @@
 
 #pragma once
 
+#include "DMXGDTFVersion.h"
 #include "Templates/SharedPointer.h"
 
 class FXmlNode;
@@ -9,17 +10,14 @@ class FXmlNode;
 namespace UE::DMX::GDTF
 {
 	class FDMXGDTFFixtureType;
+	class FDMXGDTFXmlNodeBuilder;
 
-	/**
-	 * Base class for GDTF nodes.
-	 *
-	 * The tree of GDTF nodes is constructed by calling InitializeFromXmlNode from its root. Use DMXGDTFNodeParser to parse attributes and create children.
-	 */
+	/** Base class for GDTF nodes. */
 	class DMXGDTF_API FDMXGDTFNode
 		: public TSharedFromThis<FDMXGDTFNode>
 	{
-		// Allow node initializer to set the fixture type
 		template <typename NodeType> friend class FDMXGDTFNodeInitializer;
+		friend class FDMXGDTFXmlNodeBuilder;
 
 	public:
 		virtual ~FDMXGDTFNode() {}
@@ -30,6 +28,9 @@ namespace UE::DMX::GDTF
 		/** Initializes the node from an Xml node. Called after the node was constructed. */
 		virtual void Initialize(const FXmlNode& InXmlNode) = 0;
 
+		/** Creates an XML node in the parent node */
+		virtual FXmlNode* CreateXmlNode(FXmlNode& Parent) { return nullptr; };
+
 		/** Returns the fixture type this node resides in */
 		TWeakPtr<FDMXGDTFFixtureType> GetFixtureType() const { return WeakFixtureType; }
 
diff --git a/Engine/Plugins/VirtualProduction/DMX/DMXGDTF/Source/DMXGDTF/Internal/GDTF/DMXGDTFVersion.h b/Engine/Plugins/VirtualProduction/DMX/DMXGDTF/Source/DMXGDTF/Internal/GDTF/DMXGDTFVersion.h
--- a/Engine/Plugins/VirtualProduction/DMX/DMXGDTF/Source/DMXGDTF/Internal/GDTF/DMXGDTFVersion.h
+++ b/Engine/Plugins/VirtualProduction/DMX/DMXGDTF/Source/DMXGDTF/Internal/GDTF/DMXGDTFVersion.h
@@ -5,16 +5,19 @@
 #include "Containers/UnrealString.h"
 #include "HAL/Platform.h"
 
-/** Version of the GDTF standard currently being used by the engine */
-struct DMXGDTF_API FDMXGDTFVersion
+namespace UE::DMX::GDTF
 {
-	static const uint8 MajorVersion = 1;
+	/** Version of the GDTF standard currently being used by the engine */
+	struct DMXGDTF_API FDMXGDTFVersion
+	{
+		static const uint8 MajorVersion = 1;
 
-	static const uint8 MinorVersion = 2;
+		static const uint8 MinorVersion = 1;
 
-	static const FString GetMajorVersionAsString();
-	static const FString GetMinorVersionAsString();
+		static const FString GetMajorVersionAsString();
+		static const FString GetMinorVersionAsString();
 
-	/** Returns the max supported major and minor version as string, in the format of 'Major.Minor', e.g. "1.2" */
-	static const FString GetAsString();
-};
+		/** Returns the max supported major and minor version as string, in the format of 'Major.Minor', e.g. "1.2" */
+		static const FString GetAsString();
+	};
+}
diff --git a/Engine/Plugins/VirtualProduction/DMX/DMXGDTF/Source/DMXGDTF/Internal/GDTF/DMXModes/DMXGDTFChannelFunction.h b/Engine/Plugins/VirtualProduction/DMX/DMXGDTF/Source/DMXGDTF/Internal/GDTF/DMXModes/DMXGDTFChannelFunction.h
--- a/Engine/Plugins/VirtualProduction/DMX/DMXGDTF/Source/DMXGDTF/Internal/GDTF/DMXModes/DMXGDTFChannelFunction.h
+++ b/Engine/Plugins/VirtualProduction/DMX/DMXGDTF/Source/DMXGDTF/Internal/GDTF/DMXModes/DMXGDTFChannelFunction.h
@@ -32,6 +32,7 @@ namespace UE::DMX::GDTF
 		//~ Begin FDMXGDTFNode interface
 		virtual const TCHAR* GetXmlTag() const override { return TEXT("ChannelFunction"); }
 		virtual void Initialize(const FXmlNode& XmlNode) override;
+		virtual FXmlNode* CreateXmlNode(FXmlNode& Parent) override;
 		//~ End FDMXGDTFNode interface;
 
 		/** Unique name; Default value: Name of attribute and number of channel function. */
@@ -44,25 +45,25 @@ namespace UE::DMX::GDTF
 		FString OriginalAttribute;
 
 		/**
-		 * Start DMX value; The end DMX value is calculated as a DMXFrom of the next channel function — 1 or the maximum value of the DMX channel.
-		 * Default value : “0 / 1”.
+		 * Start DMX value; The end DMX value is calculated as a DMXFrom of the next channel function -1 or the maximum value of the DMX channel.
+		 * Default value: "0/1".
 		 */
 		FDMXGDTFDMXValue DMXFrom;
 
-		/** Default DMX value of channel function when activated by the contro system. */
+		/** Default DMX value of channel function when activated by the control system. */
 		FDMXGDTFDMXValue Default;
 
 		/** Physical start value; Default value : 0 */
-		double PhysicalFrom = 0.0;
+		float PhysicalFrom = 0.f;
 
 		/** Physical end value; Default value : 1 */
-		double PhysicalTo = 1.0;
+		float PhysicalTo = 1.f;
 
 		/** Time in seconds to move from min to max of the Channel Function; Default value : 0 */
-		double RealFade = 0.0;
+		float RealFade = 0.f;
 
 		/** Time in seconds to accelerate from stop to maximum velocity; Default value : 0 */
-		double RealAcceleration = 0.0;
+		float RealAcceleration = 0.f;
 
 		/** (Optional) Link to a wheel; Starting point: Wheel Collect */
 		FString Wheel;
@@ -82,20 +83,20 @@ namespace UE::DMX::GDTF
 		/** (Optional) Link to DMX Channel or Channel Function; Starting point DMX mode. */
 		FString ModeMaster;
 
-		/** Only used together with ModeMaster; DMX start value; Default value: 0 / 1 */
-		FDMXGDTFDMXValue ModeFrom;
+		/** Only used together with ModeMaster; DMX start value; Default value: 0/1 */
+		FDMXGDTFDMXValue ModeFrom = TEXT("0/1");
 
-		/** Only used together with ModeMaster; DMX start value; Default value: 0 / 1 */
-		FDMXGDTFDMXValue ModeTo;
+		/** Only used together with ModeMaster; DMX start value; Default value: 0/1 */
+		FDMXGDTFDMXValue ModeTo = TEXT("0/1");
 
 		/** (Optional) Link to DMX Profile; Starting point: DMX Profile Collect */
 		FString DMXProfile;
 
 		/** Minimum Physical Value that will be used for the DMX Profile. Default: Value from PhysicalFrom 1 */
-		double Min = 0.0;
+		float Min = 0.f;
 
 		/**  Maximum Physical Value that will be used for the DMX Profile. Default: Value from PhysicalTo */
-		double Max = 0.0;
+		float Max = 1.f;
 
 		/**
 		 * Custom Name that can he used do adress this channel function with
@@ -142,5 +143,8 @@ namespace UE::DMX::GDTF
 
 		/** Resolves the linked DMX profile. Returns the DMX profile, or nullptr if no DMX profile is linked */
 		TSharedPtr<FDMXGDTFDMXProfile> ResolveDMXProfile() const;
+
+		/** Parses the default value. Useful to parse legacy GDTFs that store the default in the DMX channel node */
+		FDMXGDTFDMXValue ParseDefault(const FString& Value, const FXmlNode* XmlNode) const;
 	};
 }
diff --git a/Engine/Plugins/VirtualProduction/DMX/DMXGDTF/Source/DMXGDTF/Internal/GDTF/DMXModes/DMXGDTFChannelSet.h b/Engine/Plugins/VirtualProduction/DMX/DMXGDTF/Source/DMXGDTF/Internal/GDTF/DMXModes/DMXGDTFChannelSet.h
--- a/Engine/Plugins/VirtualProduction/DMX/DMXGDTF/Source/DMXGDTF/Internal/GDTF/DMXModes/DMXGDTFChannelSet.h
+++ b/Engine/Plugins/VirtualProduction/DMX/DMXGDTF/Source/DMXGDTF/Internal/GDTF/DMXModes/DMXGDTFChannelSet.h
@@ -19,6 +19,7 @@ namespace UE::DMX::GDTF
 		//~ Begin FDMXGDTFNode interface
 		virtual const TCHAR* GetXmlTag() const override { return TEXT("ChannelSet"); }
 		virtual void Initialize(const FXmlNode& XmlNode) override;
+		virtual FXmlNode* CreateXmlNode(FXmlNode& Parent) override;
 		//~ End FDMXGDTFNode interface;
 
 		/** The name of the channel set. Default: Empty */
@@ -32,10 +33,10 @@ namespace UE::DMX::GDTF
 		FDMXGDTFDMXValue DMXFrom;
 
 		/** Physical start value */
-		float PhysicalFrom = 0.0;
+		float PhysicalFrom = 0.f;
 
 		/** Physical end value */
-		float PhysicalTo = 0.0;
+		float PhysicalTo = 1.f;
 
 		/**
 		 * If the channel function has a link to a wheel, a corresponding slot index
@@ -43,7 +44,7 @@ namespace UE::DMX::GDTF
 		 * the wheel which is linked in the channel function. The wheel slot index is
 		 * normalized to 1. Size: 4 bytes
 		 */
-		uint32 WheelSlotIndex = INDEX_NONE;
+		int32 WheelSlotIndex = INDEX_NONE;
 
 		/** The outer channel function */
 		const TWeakPtr<FDMXGDTFChannelFunction> OuterChannelFunction;
diff --git a/Engine/Plugins/VirtualProduction/DMX/DMXGDTF/Source/DMXGDTF/Internal/GDTF/DMXModes/DMXGDTFDMXChannel.h b/Engine/Plugins/VirtualProduction/DMX/DMXGDTF/Source/DMXGDTF/Internal/GDTF/DMXModes/DMXGDTFDMXChannel.h
--- a/Engine/Plugins/VirtualProduction/DMX/DMXGDTF/Source/DMXGDTF/Internal/GDTF/DMXModes/DMXGDTFDMXChannel.h
+++ b/Engine/Plugins/VirtualProduction/DMX/DMXGDTF/Source/DMXGDTF/Internal/GDTF/DMXModes/DMXGDTFDMXChannel.h
@@ -25,15 +25,16 @@ namespace UE::DMX::GDTF
 		//~ Begin FDMXGDTFNode interface
 		virtual const TCHAR* GetXmlTag() const override { return TEXT("DMXChannel"); }
 		virtual void Initialize(const FXmlNode& XmlNode) override;
+		virtual FXmlNode* CreateXmlNode(FXmlNode& Parent) override;
 		//~ End FDMXGDTFNode interface;
 
 		/**
 		 * Number of the DMXBreak; Default value: 1;
 		 * Special value: “Overwrite” – means that this number will be overwritten by Geometry Reference; Size: 4 bytes
 		 *
-		 * @todo UE Specific: If the optional is not set and a related geometry reference exists, it means the value is set to special value "Overwrite".
+		 * UE Specific: If negative and a related geometry reference exists, it means the value is set to special value "Overwrite".
 		 */
-		uint32 DMXBreak;
+		int32 DMXBreak = 1;
 
 		/**
 		 * Relative addresses of the current DMX channel from highest to least significant; Size per int: 4 bytes
@@ -49,7 +50,7 @@ namespace UE::DMX::GDTF
 		FString InitialFunction;
 
 		/** Highlight value for current channel; Special value : “None”.Default value : “None” */
-		FDMXGDTFDMXValue DMXValue;
+		FDMXGDTFDMXValue Highlight;
 
 		/**
 		 * Name of the geometry the current channel controls.
@@ -66,6 +67,9 @@ namespace UE::DMX::GDTF
 		/** The outer DMX mode */
 		const TWeakPtr<FDMXGDTFDMXMode> OuterDMXMode;
 
+		UE_DEPRECATED(5.5, "Deprecated with GDTF 1.1. Instead each channel function can hold its own default. Please refer to DMXGDTFChannelFunction::Default")
+		FDMXGDTFDMXValue Default;
+
 		/** Resolves the linked initial function. Returns the initial function, or nullptr if no initial function is linked */
 		TSharedPtr<FDMXGDTFChannelFunction> ResolveInitialFunction() const;
 
diff --git a/Engine/Plugins/VirtualProduction/DMX/DMXGDTF/Source/DMXGDTF/Internal/GDTF/DMXModes/DMXGDTFDMXMode.h b/Engine/Plugins/VirtualProduction/DMX/DMXGDTF/Source/DMXGDTF/Internal/GDTF/DMXModes/DMXGDTFDMXMode.h
--- a/Engine/Plugins/VirtualProduction/DMX/DMXGDTF/Source/DMXGDTF/Internal/GDTF/DMXModes/DMXGDTFDMXMode.h
+++ b/Engine/Plugins/VirtualProduction/DMX/DMXGDTF/Source/DMXGDTF/Internal/GDTF/DMXModes/DMXGDTFDMXMode.h
@@ -22,6 +22,7 @@ namespace UE::DMX::GDTF
 		//~ Begin FDMXGDTFNode interface
 		virtual const TCHAR* GetXmlTag() const override { return TEXT("DMXMode"); }
 		virtual void Initialize(const FXmlNode& XmlNode) override;
+		virtual FXmlNode* CreateXmlNode(FXmlNode& Parent) override;
 		//~ End FDMXGDTFNode interface
 
 		/** The unique name of the DMX mode */
@@ -30,11 +31,7 @@ namespace UE::DMX::GDTF
 		/** Description of the DMX mode. */
 		FString Description;
 
-		/**
-		 * The first geometry in the device; Only top level geometries are allowed to be linked
-		 *
-		 * UE Specific: A ptr to the geometry instead of its unique name.
-		 */
+		/** The first geometry in the device; Only top level geometries are allowed to be linked	 */
 		FString Geometry;
 
 		/** Description of all DMX channels used in the mode */
diff --git a/Engine/Plugins/VirtualProduction/DMX/DMXGDTF/Source/DMXGDTF/Internal/GDTF/DMXModes/DMXGDTFDMXValue.h b/Engine/Plugins/VirtualProduction/DMX/DMXGDTF/Source/DMXGDTF/Internal/GDTF/DMXModes/DMXGDTFDMXValue.h
--- a/Engine/Plugins/VirtualProduction/DMX/DMXGDTF/Source/DMXGDTF/Internal/GDTF/DMXModes/DMXGDTFDMXValue.h
+++ b/Engine/Plugins/VirtualProduction/DMX/DMXGDTF/Source/DMXGDTF/Internal/GDTF/DMXModes/DMXGDTFDMXValue.h
@@ -14,7 +14,10 @@ namespace UE::DMX::GDTF
 	*/
 	struct DMXGDTF_API FDMXGDTFDMXValue
 	{
-		FString Value;
+		FDMXGDTFDMXValue() = default;
+		FDMXGDTFDMXValue(const TCHAR* InValue);
+
+		FString Value = TEXT("None");
 
 		/** Returns the value as integer. Note this is a relatively slow operation as it parses the string. */
 		bool ToInt(uint32& OutInteger) const;
diff --git a/Engine/Plugins/VirtualProduction/DMX/DMXGDTF/Source/DMXGDTF/Internal/GDTF/DMXModes/DMXGDTFLogicalChannel.h b/Engine/Plugins/VirtualProduction/DMX/DMXGDTF/Source/DMXGDTF/Internal/GDTF/DMXModes/DMXGDTFLogicalChannel.h
--- a/Engine/Plugins/VirtualProduction/DMX/DMXGDTF/Source/DMXGDTF/Internal/GDTF/DMXModes/DMXGDTFLogicalChannel.h
+++ b/Engine/Plugins/VirtualProduction/DMX/DMXGDTF/Source/DMXGDTF/Internal/GDTF/DMXModes/DMXGDTFLogicalChannel.h
@@ -64,6 +64,7 @@ namespace UE::DMX::GDTF
 		//~ Begin FDMXGDTFNode interface
 		virtual const TCHAR* GetXmlTag() const override { return TEXT("LogicalChannel"); }
 		virtual void Initialize(const FXmlNode& XmlNode) override;
+		virtual FXmlNode* CreateXmlNode(FXmlNode& Parent) override;
 		//~ End FDMXGDTFNode interface;
 
 		/** Link to the attribute; The starting point is the Attribute Collect. */
@@ -76,14 +77,14 @@ namespace UE::DMX::GDTF
 		EDMXGDTFLogicalChannelMaster Master = EDMXGDTFLogicalChannelMaster::None;
 
 		/** Minimum fade time for moves in black action. MibFade is defined for the complete DMX range. Default value : 0; Unit: second */
-		double MibFade = 0.0;
+		float MibFade = 0.f;
 
 		/**
 		 * Minimum fade time for the subordinate channel functions to change
 		 * DMX values by the control system.DMXChangeTimeLimit is defined
 		 * for the complete DMX range.Default value : 0; Unit: second
 		 */
-		double DMXChangeTimeLimit = 0.0;
+		float DMXChangeTimeLimit = 0.f;
 
 		/** A list of channel functions */
 		TArray<TSharedPtr<FDMXGDTFChannelFunction>> ChannelFunctionArray;
diff --git a/Engine/Plugins/VirtualProduction/DMX/DMXGDTF/Source/DMXGDTF/Internal/GDTF/DMXModes/DMXGDTFSubchannelSet.h b/Engine/Plugins/VirtualProduction/DMX/DMXGDTF/Source/DMXGDTF/Internal/GDTF/DMXModes/DMXGDTFSubchannelSet.h
--- a/Engine/Plugins/VirtualProduction/DMX/DMXGDTF/Source/DMXGDTF/Internal/GDTF/DMXModes/DMXGDTFSubchannelSet.h
+++ b/Engine/Plugins/VirtualProduction/DMX/DMXGDTF/Source/DMXGDTF/Internal/GDTF/DMXModes/DMXGDTFSubchannelSet.h
@@ -21,16 +21,17 @@ namespace UE::DMX::GDTF
 		//~ Begin FDMXGDTFNode interface
 		virtual const TCHAR* GetXmlTag() const override { return TEXT("SubchannelSet"); }
 		virtual void Initialize(const FXmlNode& XmlNode) override;
+		virtual FXmlNode* CreateXmlNode(FXmlNode& Parent) override;
 		//~ End FDMXGDTFNode interface
 
 		/** The name of the sub channel set. Default: Empty */
 		FName Name;
 
 		/** Physical start value */
-		float PhysicalFrom = 0.0;
+		float PhysicalFrom = 0.f;
 
 		/** Physical end value */
-		float PhysicalTo = 0.0;
+		float PhysicalTo = 1.f;
 
 		/** (Optional) Link to the sub physical unit; Starting Point: Attribute */
 		FString SubphyiscalUnit;
diff --git a/Engine/Plugins/VirtualProduction/DMX/DMXGDTF/Source/DMXGDTF/Internal/GDTF/Geometries/DMXGDTFBeamGeometry.h b/Engine/Plugins/VirtualProduction/DMX/DMXGDTF/Source/DMXGDTF/Internal/GDTF/Geometries/DMXGDTFBeamGeometry.h
--- a/Engine/Plugins/VirtualProduction/DMX/DMXGDTF/Source/DMXGDTF/Internal/GDTF/Geometries/DMXGDTFBeamGeometry.h
+++ b/Engine/Plugins/VirtualProduction/DMX/DMXGDTF/Source/DMXGDTF/Internal/GDTF/Geometries/DMXGDTFBeamGeometry.h
@@ -18,7 +18,7 @@ enum class EDMXGDTFBeamGeometryLampType : uint8
 	Tungsten UMETA(DisplayName = "Tungsten"),
 	Halogen UMETA(DisplayName = "Halogen"),
 	LED UMETA(DisplayName = "LED"),
-
+	
 	// Max value. Add further elements above.
 	MaxEnumValue UMETA(Hidden)
 };
@@ -45,7 +45,7 @@ enum class EDMXGDTFBeamGeometryBeamType : uint8
 	PC UMETA(DisplayName = "PC"),
 	Fresnel UMETA(DisplayName = "Fresnel"),
 	Glow UMETA(DisplayName = "Glow"),
-
+	
 	// Max value. Add further elements above.
 	MaxEnumValue UMETA(Hidden)
 };
@@ -71,34 +71,35 @@ namespace UE::DMX::GDTF
 		//~ Begin DMXGDTFGeneralGeometryNode interface
 		virtual const TCHAR* GetXmlTag() const override { return TEXT("Beam"); }
 		virtual void Initialize(const FXmlNode& XmlNode) override;
+		virtual FXmlNode* CreateXmlNode(FXmlNode& Parent) override;
 		//~ End DMXGDTFGeneralGeometryNode interface
 
 		/** Defines type of the light source */
 		EDMXGDTFBeamGeometryLampType LampType = EDMXGDTFBeamGeometryLampType::Discharge;
 
 		/** Power consumption; Default value: 1 000; Unit: Watt */
-		float PowerConsumption = 1000.0;
+		float PowerConsumption = 1000.f;
 
 		/** Intensity of all the represented light emitters; Default value: 10 000; Unit: lumen */
-		float LuminousFlux = 1000.0;
+		float LuminousFlux = 1000.f;
 
 		/** Color temperature; Default value: 6 000; Unit: kelvin */
-		float ColorTemperature = 1000.0;
+		float ColorTemperature = 1000.f;
 
 		/** Beam angle; Default value: 25,0; Unit: degree */
-		float BeamAngle = 25.0;
+		float BeamAngle = 25.f;
 
 		/** Field angle; Default value: 25,0; Unit: degree */
-		float FieldAngle = 25.0;
+		float FieldAngle = 25.f;
 
 		/** Throw Ratio of the lens for BeamType Rectangle; Default value: 1; Unit: None */
-		float ThrowRatio = 1.0;
+		float ThrowRatio = 1.f;
 
 		/** Ratio from Width to Height of the Rectangle Type Beam; Default value: 1.7777; Unit: None */
-		float RectangleRatio = 1.7777;
+		float RectangleRatio = 1.7777f;
 
 		/** Beam radius on starting point. Default value: 0,05; Unit: meter. */
-		float BeamRadius = 0.05;
+		float BeamRadius = 0.05f;
 
 		/**
 		 * Beam Type
diff --git a/Engine/Plugins/VirtualProduction/DMX/DMXGDTF/Source/DMXGDTF/Internal/GDTF/Geometries/DMXGDTFBreak.h b/Engine/Plugins/VirtualProduction/DMX/DMXGDTF/Source/DMXGDTF/Internal/GDTF/Geometries/DMXGDTFBreak.h
--- a/Engine/Plugins/VirtualProduction/DMX/DMXGDTF/Source/DMXGDTF/Internal/GDTF/Geometries/DMXGDTFBreak.h
+++ b/Engine/Plugins/VirtualProduction/DMX/DMXGDTF/Source/DMXGDTF/Internal/GDTF/Geometries/DMXGDTFBreak.h
@@ -3,7 +3,7 @@
 #pragma once
 
 #include "GDTF/DMXGDTFNode.h"
-#include "GDTF/Geometries/DMXGDTFAddress.h"
+#include "GDTF/Geometries/DMXGDTFDMXAddress.h"
 
 namespace UE::DMX::GDTF
 {
@@ -19,10 +19,11 @@ namespace UE::DMX::GDTF
 		//~ Begin DMXGDTFGeneralGeometryNode interface
 		virtual const TCHAR* GetXmlTag() const override { return TEXT("Break"); }
 		virtual void Initialize(const FXmlNode& XmlNode) override;
+		virtual FXmlNode* CreateXmlNode(FXmlNode& Parent) override;
 		//~ End DMXGDTFGeneralGeometryNode interface
 
 		/** DMX offset; Default value:1 (Means no offset for the corresponding DMX Channel) */
-		FDMXGDTFAddress DMXAddress;
+		FDMXGDTFDMXAddress DMXOffset;
 
 		/** Defines the unique number of the DMX Break for which the Offset is given. Size: 1 byte; Default value 1 */
 		uint8 DMXBreak = 1;
diff --git a/Engine/Plugins/VirtualProduction/DMX/DMXGDTF/Source/DMXGDTF/Internal/GDTF/Geometries/DMXGDTFGeometryCollectBase.h b/Engine/Plugins/VirtualProduction/DMX/DMXGDTF/Source/DMXGDTF/Internal/GDTF/Geometries/DMXGDTFGeometryCollectBase.h
--- a/Engine/Plugins/VirtualProduction/DMX/DMXGDTF/Source/DMXGDTF/Internal/GDTF/Geometries/DMXGDTFGeometryCollectBase.h
+++ b/Engine/Plugins/VirtualProduction/DMX/DMXGDTF/Source/DMXGDTF/Internal/GDTF/Geometries/DMXGDTFGeometryCollectBase.h
@@ -36,8 +36,10 @@ namespace UE::DMX::GDTF
 		//~ Begin FDMXGDTFNode interface
 		virtual const TCHAR* GetXmlTag() const override;
 		virtual void Initialize(const FXmlNode& XmlNode) override;
+		virtual FXmlNode* CreateXmlNode(FXmlNode& Parent) override;
 		//~ End FDMXGDTFNode interface
 
+	public:
 		/** Any General Geometry. */
 		TArray<TSharedPtr<FDMXGDTFGeometry>> GeometryArray;
 
@@ -95,13 +97,13 @@ namespace UE::DMX::GDTF
 		/** Finds the geometry by name. Can be either a geometry or a geometry reference */
 		void FindGeometryByName(const TCHAR* InName, TSharedPtr<FDMXGDTFGeometry>& OutGeometry, TSharedPtr<FDMXGDTFGeometryReference>& OutGeometryReference) const;
 
-		template <typename GeometryType, typename ForwardDMXGDTFGeometry = FDMXGDTFGeometry, typename ForwardDMXGDTFGeometryReference = FDMXGDTFGeometryReference>
+		template <typename GeometryType>
 		TSharedPtr<GeometryType> ResolveGeometryLink(const FString& Link) const
 		{
 			TArray<FString> LinkArray;
 			Link.ParseIntoArray(LinkArray, TEXT("."));
 
-			TSharedPtr<const FDMXGDTFGeometryCollectBase> OuterGeometryCollect = SharedThis(this);
+			TSharedPtr<const FDMXGDTFGeometryCollectBase> OuterGeometryCollect = StaticCastSharedRef<const FDMXGDTFGeometryCollectBase>(SharedThis(this));
 			for (const FString& GeometryName : LinkArray)
 			{
 				if (&LinkArray.Last() == &GeometryName)
@@ -115,12 +117,12 @@ namespace UE::DMX::GDTF
 					}
 				}
 
-				TSharedPtr<ForwardDMXGDTFGeometry> Geometry;
-				TSharedPtr<ForwardDMXGDTFGeometryReference> GeometryReference;
+				TSharedPtr<FDMXGDTFGeometry> Geometry;
+				TSharedPtr<FDMXGDTFGeometryReference> GeometryReference;
 				OuterGeometryCollect->FindGeometryByName(*GeometryName, Geometry, GeometryReference);
 				if (Geometry.IsValid())
 				{
-					OuterGeometryCollect = Geometry;
+					OuterGeometryCollect = StaticCastSharedPtr<const FDMXGDTFGeometryCollectBase>(Geometry);
 				}
 			}
 
diff --git a/Engine/Plugins/VirtualProduction/DMX/DMXGDTF/Source/DMXGDTF/Internal/GDTF/Geometries/DMXGDTFGeometryReference.h b/Engine/Plugins/VirtualProduction/DMX/DMXGDTF/Source/DMXGDTF/Internal/GDTF/Geometries/DMXGDTFGeometryReference.h
--- a/Engine/Plugins/VirtualProduction/DMX/DMXGDTF/Source/DMXGDTF/Internal/GDTF/Geometries/DMXGDTFGeometryReference.h
+++ b/Engine/Plugins/VirtualProduction/DMX/DMXGDTF/Source/DMXGDTF/Internal/GDTF/Geometries/DMXGDTFGeometryReference.h
@@ -4,6 +4,8 @@
 
 #include "GDTF/DMXGDTFNode.h"
 
+#include "Math/Transform.h" 
+
 namespace UE::DMX::GDTF
 {
 	class FDMXGDTFBreak;
@@ -23,11 +25,15 @@ namespace UE::DMX::GDTF
 		//~ Begin DMXGDTFGeneralGeometryNode interface
 		virtual const TCHAR* GetXmlTag() const override { return TEXT("GeometryReference"); }
 		virtual void Initialize(const FXmlNode& XmlNode) override;
+		virtual FXmlNode* CreateXmlNode(FXmlNode& Parent) override;
 		//~ End DMXGDTFGeneralGeometryNode interface
 
 		/** The unique name of geometry. */
 		FName Name;
 
+		/** Relative position of geometry; Default value : Identity Matrix */
+		FTransform Position;
+
 		/** The referenced geometry. Only top level geometries are allowed to be referenced */
 		FString Geometry;
 
diff --git a/Engine/Plugins/VirtualProduction/DMX/DMXGDTF/Source/DMXGDTF/Internal/GDTF/Geometries/DMXGDTFLaserGeometry.h b/Engine/Plugins/VirtualProduction/DMX/DMXGDTF/Source/DMXGDTF/Internal/GDTF/Geometries/DMXGDTFLaserGeometry.h
--- a/Engine/Plugins/VirtualProduction/DMX/DMXGDTF/Source/DMXGDTF/Internal/GDTF/Geometries/DMXGDTFLaserGeometry.h
+++ b/Engine/Plugins/VirtualProduction/DMX/DMXGDTF/Source/DMXGDTF/Internal/GDTF/Geometries/DMXGDTFLaserGeometry.h
@@ -37,37 +37,38 @@ namespace UE::DMX::GDTF
 		//~ Begin DMXGDTFGeneralGeometryNode interface
 		virtual const TCHAR* GetXmlTag() const override { return TEXT("Laser"); }
 		virtual void Initialize(const FXmlNode& XmlNode) override;
+		virtual FXmlNode* CreateXmlNode(FXmlNode& Parent) override;
 		//~ End DMXGDTFGeneralGeometryNode interface
 
 		/** Color type of a Laser */
 		EDMXLaserGeometryColorType ColorType = EDMXLaserGeometryColorType::RGB;
 
 		/** Required if ColorType is “SingleWaveLength”; Unit:nm (nanometers) */
-		double Color = 0.f;
+		float Color = 0.f;
 
 		/** Output Strength of the Laser; Unit: Watt */
-		double OutputStrength = 0.f;
+		float OutputStrength = 0.f;
 
 		/** (Optional) link to the emitter group. The starting point is the Emitter Collect. */
 		FString Emitter;
 
 		/** Beam diameter where it leaves the projector; Unit: meter */
-		double BeamDiameter = 0.f;
+		float BeamDiameter = 0.f;
 
 		/** Minimum beam divergence; Unit: mrad (milliradian) */
-		double BeamDivergenceMin = 0.f;
+		float BeamDivergenceMin = 0.f;
 
 		/** Maximum beam divergence; Unit: mrad (milliradian) */
-		double BeamDivergenceMax = 0.f;
+		float BeamDivergenceMax = 0.f;
 
 		/** Possible Total Scan Angle Pan of the beam. Assumes symmetrical output; Unit: Degree */
-		double ScanAnglePan = 0.f;
+		float ScanAnglePan = 0.f;
 
 		/** Possible Total Scan Angle Tilt of the beam. Assumes symmetrical output; Unit: Degree */
-		double ScanAngleTilt = 0.f;
+		float ScanAngleTilt = 0.f;
 
 		/** Speed of the beam; Unit: kilo point per second */
-		double ScanSpeed = 0.f;
+		float ScanSpeed = 0.f;
 
 		/** A list of protocols supported by the laser */
 		TArray<TSharedPtr<FDMXGDTFLaserProtocol>> ProtocolArray;
diff --git a/Engine/Plugins/VirtualProduction/DMX/DMXGDTF/Source/DMXGDTF/Internal/GDTF/Geometries/DMXGDTFStructureGeometry.h b/Engine/Plugins/VirtualProduction/DMX/DMXGDTF/Source/DMXGDTF/Internal/GDTF/Geometries/DMXGDTFStructureGeometry.h
--- a/Engine/Plugins/VirtualProduction/DMX/DMXGDTF/Source/DMXGDTF/Internal/GDTF/Geometries/DMXGDTFStructureGeometry.h
+++ b/Engine/Plugins/VirtualProduction/DMX/DMXGDTF/Source/DMXGDTF/Internal/GDTF/Geometries/DMXGDTFStructureGeometry.h
@@ -48,6 +48,7 @@ namespace UE::DMX::GDTF
 		//~ Begin DMXGDTFGeneralGeometryNode interface
 		virtual const TCHAR* GetXmlTag() const override { return TEXT("Structure"); }
 		virtual void Initialize(const FXmlNode& XmlNode) override;
+		virtual FXmlNode* CreateXmlNode(FXmlNode& Parent) override;
 		//~ End DMXGDTFGeneralGeometryNode interface
 
 		/** The linked geometry */
@@ -60,10 +61,10 @@ namespace UE::DMX::GDTF
 		EDMXGDTFStructureGeometryCrossSectionType CrossSectionType = EDMXGDTFStructureGeometryCrossSectionType::TrussFramework;
 
 		/** The height of the cross section. Only for Tubes. Unit: meter */
-		double CrossSectionHeight = 0.f;
+		float CrossSectionHeight = 0.f;
 
 		/** The thickness of the wall of the cross section. Only for Tubes. Unit: meter. */
-		double CrossSectionWallThickness = 0.f;
+		float CrossSectionWallThickness = 0.f;
 
 		/** The name of the truss cross section. Only for Trusses. */
 		FString TrussCrossSection;
@@ -71,4 +72,4 @@ namespace UE::DMX::GDTF
 		/** Resolves the linked geometry. Returns the geometry, or nullptr if no geometry is linked */
 		TSharedPtr<FDMXGDTFGeometry> ResolveLinkedGeometry() const;
 	};
-}
\ No newline at end of file
+}
diff --git a/Engine/Plugins/VirtualProduction/DMX/DMXGDTF/Source/DMXGDTF/Internal/GDTF/Geometries/DMXGDTFSupportGeometry.h b/Engine/Plugins/VirtualProduction/DMX/DMXGDTF/Source/DMXGDTF/Internal/GDTF/Geometries/DMXGDTFSupportGeometry.h
--- a/Engine/Plugins/VirtualProduction/DMX/DMXGDTF/Source/DMXGDTF/Internal/GDTF/Geometries/DMXGDTFSupportGeometry.h
+++ b/Engine/Plugins/VirtualProduction/DMX/DMXGDTF/Source/DMXGDTF/Internal/GDTF/Geometries/DMXGDTFSupportGeometry.h
@@ -34,6 +34,7 @@ namespace UE::DMX::GDTF
 		//~ Begin DMXGDTFGeneralGeometryNode interface
 		virtual const TCHAR* GetXmlTag() const override { return TEXT("Support"); }
 		virtual void Initialize(const FXmlNode& XmlNode) override;
+		virtual FXmlNode* CreateXmlNode(FXmlNode& Parent) override;
 		//~ End DMXGDTFGeneralGeometryNode interface
 
 		/** The type of support */
@@ -46,39 +47,39 @@ namespace UE::DMX::GDTF
 		FVector RopeOffset;
 
 		/** The allowable force on the X-Axis applied to the object according to the Eurocode. Unit: N. */
-		double CapacityX = 0.f;
+		float CapacityX = 0.f;
 
 		/** The allowable force on the Y-Axis applied to the object according to the Eurocode. Unit: N. */
-		double CapacityY = 0.f;
+		float CapacityY = 0.f;
 
 		/** The allowable force on the Z-Axis applied to the object according to the Eurocode. Unit: N. */
-		double CapacityZ = 0.f;
+		float CapacityZ = 0.f;
 
 		/** The allowable moment around the X-Axis applied to the object according to the Eurocode. Unit: N/m. */
-		double CapacityXX = 0.f;
+		float CapacityXX = 0.f;
 
 		/** The allowable moment around the Y-Axis applied to the object according to the Eurocode. Unit: N/m. */
-		double CapacityYY = 0.f;
+		float CapacityYY = 0.f;
 
 		/** The allowable moment around the Z-Axis applied to the object according to the Eurocode. Unit: N/m. */
-		double CapacityZZ = 0.f;
+		float CapacityZZ = 0.f;
 
 		/** The compression ratio for this support along the X-Axis. Unit N/m. Only for Ground Supports. */
-		double ResistanceX = 0.f;
+		float ResistanceX = 0.f;
 
 		/** The compression ratio for this support along the Y-Axis. Unit N/m. Only for Ground Supports. */
-		double ResistanceY = 0.f;
+		float ResistanceY = 0.f;
 
 		/** The compression ratio for this support along the Z-Axis. Unit N/m. Only for Ground Supports. */
-		double ResistanceZ = 0.f;
+		float ResistanceZ = 0.f;
 
 		/** The compression ratio for this support around the X-Axis. Unit N/m. Only for Ground Supports. */
-		double ResistanceXX = 0.f;
+		float ResistanceXX = 0.f;
 
 		/** The compression ratio for this support around the Y-Axis. Unit N/m. Only for Ground Supports. */
-		double ResistanceYY = 0.f;
+		float ResistanceYY = 0.f;
 
 		/** The compression ratio for this support around the Z-Axis. Unit N/m. Only for Ground Supports. */
-		double ResistanceZZ = 0.f;
+		float ResistanceZZ = 0.f;
 	};
 }
diff --git a/Engine/Plugins/VirtualProduction/DMX/DMXGDTF/Source/DMXGDTF/Internal/GDTF/Geometries/DMXGDTFWiringObjectGeometry.h b/Engine/Plugins/VirtualProduction/DMX/DMXGDTF/Source/DMXGDTF/Internal/GDTF/Geometries/DMXGDTFWiringObjectGeometry.h
--- a/Engine/Plugins/VirtualProduction/DMX/DMXGDTF/Source/DMXGDTF/Internal/GDTF/Geometries/DMXGDTFWiringObjectGeometry.h
+++ b/Engine/Plugins/VirtualProduction/DMX/DMXGDTF/Source/DMXGDTF/Internal/GDTF/Geometries/DMXGDTFWiringObjectGeometry.h
@@ -76,6 +76,7 @@ namespace UE::DMX::GDTF
 		//~ Begin DMXGDTFGeneralGeometryNode interface
 		virtual const TCHAR* GetXmlTag() const override { return TEXT("WiringObject"); }
 		virtual void Initialize(const FXmlNode& XmlNode) override;
+		virtual FXmlNode* CreateXmlNode(FXmlNode& Parent) override;
 		//~ End DMXGDTFGeneralGeometryNode interface
 
 		/** Link to the corresponding model. */
@@ -105,25 +106,25 @@ namespace UE::DMX::GDTF
 		int32 PinCount = 0;
 
 		/** The electrical consumption in Watts.Only for Consumers.Unit: Watt. */
-		double ElectricalPayLoad = 0.f;
+		float ElectricalPayLoad = 0.f;
 
 		/** The voltage range’s maximum value.Only for Consumers.Unit: volt. */
-		double VoltageRangeMax = 0.f;
+		float VoltageRangeMax = 0.f;
 
 		/** The voltage range’s minimum value.Only for Consumers.Unit: volt. */
-		double VoltageRangeMin = 0.f;
+		float VoltageRangeMin = 0.f;
 
 		/** The Frequency range’s maximum value.Only for Consumers.Unit: hertz. */
-		double FrequencyRangeMax = 0.f;
+		float FrequencyRangeMax = 0.f;
 
 		/** The Frequency range’s minimum value.Only for Consumers.Unit: hertz. */
-		double FrequencyRangeMin = 0.f;
+		float FrequencyRangeMin = 0.f;
 
 		/** The maximum electrical payload that this power source can handle. Only for Power Sources. Unit: voltampere.  */
-		double MaxPayLoad = 0.f;
+		float MaxPayLoad = 0.f;
 
 		/** he voltage output that this power source can handle. Only for Power Sources.Unit : volt. */
-		double Voltage = 0.f;
+		float Voltage = 0.f;
 
 		/**
 		 * The layer of the Signal Type. In one device, all wiring geometry that
@@ -133,10 +134,10 @@ namespace UE::DMX::GDTF
 		int32 SignalLayer = 0;
 
 		/** The Power Factor of the device. Only for consumers. */
-		double CosPhi = 0.f;
+		float CosPhi = 0.f;
 
 		/** The fuse value.Only for fuses. Unit: ampere.  */
-		double FuseCurrent = 0.f;
+		float FuseCurrent = 0.f;
 
 		/** Fuse Rating */
 		EDMXWiringObjectGeometryFuseRating FuseRating = EDMXWiringObjectGeometryFuseRating::B;
diff --git a/Engine/Plugins/VirtualProduction/DMX/DMXGDTF/Source/DMXGDTF/Internal/GDTF/Models/DMXGDTFModel.h b/Engine/Plugins/VirtualProduction/DMX/DMXGDTF/Source/DMXGDTF/Internal/GDTF/Models/DMXGDTFModel.h
--- a/Engine/Plugins/VirtualProduction/DMX/DMXGDTF/Source/DMXGDTF/Internal/GDTF/Models/DMXGDTFModel.h
+++ b/Engine/Plugins/VirtualProduction/DMX/DMXGDTF/Source/DMXGDTF/Internal/GDTF/Models/DMXGDTFModel.h
@@ -54,19 +54,20 @@ namespace UE::DMX::GDTF
 		//~ Begin FDMXGDTFNode interface
 		virtual const TCHAR* GetXmlTag() const override { return TEXT("Model"); }
 		virtual void Initialize(const FXmlNode& XmlNode) override;
+		virtual FXmlNode* CreateXmlNode(FXmlNode& Parent) override;
 		//~ End FDMXGDTFNode interface
 
 		/** The unique name of the model. */
 		FName Name;
 
 		/** Length. Unit: meter; Default value: 0 */
-		double Length = 0.0;
+		float Length = 0.0;
 
 		/** Width. Unit: meter; Default value: 0 */
-		double Width = 0.0;
+		float Width = 0.0;
 
 		/** Height. Unit: meter; Default value: 0 */
-		double Height = 0.0;
+		float Height = 0.0;
 
 		/** Type of 3D model. Default value: “Undefined */
 		EDMXGDTFModelPrimitiveType PrimitiveType = EDMXGDTFModelPrimitiveType::Undefined;
diff --git a/Engine/Plugins/VirtualProduction/DMX/DMXGDTF/Source/DMXGDTF/Internal/GDTF/PhysicalDescriptions/DMXGDTFColorSpace.h b/Engine/Plugins/VirtualProduction/DMX/DMXGDTF/Source/DMXGDTF/Internal/GDTF/PhysicalDescriptions/DMXGDTFColorSpace.h
--- a/Engine/Plugins/VirtualProduction/DMX/DMXGDTF/Source/DMXGDTF/Internal/GDTF/PhysicalDescriptions/DMXGDTFColorSpace.h
+++ b/Engine/Plugins/VirtualProduction/DMX/DMXGDTF/Source/DMXGDTF/Internal/GDTF/PhysicalDescriptions/DMXGDTFColorSpace.h
@@ -2,6 +2,7 @@
 
 #pragma once
 
+#include "DMXGDTFColorCIE1931xyY.h"
 #include "GDTF/DMXGDTFNode.h"
 #include "Math/Vector.h"
 #include "Misc/EnumRange.h"
@@ -42,13 +43,14 @@ namespace UE::DMX::GDTF
 		//~ Begin FDMXGDTFNode interface
 		virtual const TCHAR* GetXmlTag() const override { return TEXT("ColorSpace"); }
 		virtual void Initialize(const FXmlNode& XmlNode) override;
+		virtual FXmlNode* CreateXmlNode(FXmlNode& Parent) override;
 		//~ End FDMXGDTFNode interface
 		
 		/** 
 		 * Unique Name of the Color Space. Default Value : "Default".
 		 * Note that the name need to be unique for the default colorspace and all color spaces in the AdditionalColorSpaces node. 
 		 */
-		FName Name = "Default";
+		FName Name = "";
 
 		/** 
 		 * Definition of the Color Space that used for the indirect color mixing.
@@ -57,16 +59,16 @@ namespace UE::DMX::GDTF
 		EDMXGDTFColorSpaceMode Mode = EDMXGDTFColorSpaceMode::sRGB;
 
 		/** (Optional) CIE xyY of the Red Primary; this is used only if the ColorSpace is "Custom". */
-		FVector Red = FVector::ZeroVector;	
+		FDMXGDTFColorCIE1931xyY Red;	
 		
 		/** (Optional) CIE xyY of the Green Primary; this is used only if the ColorSpace is "Custom". */
-		FVector Green = FVector::ZeroVector;
+		FDMXGDTFColorCIE1931xyY Green;
 
 		/** (Optional) CIE xyY of the Blue Primary; this is used only if the ColorSpace is "Custom". */
-		FVector Blue = FVector::ZeroVector;
+		FDMXGDTFColorCIE1931xyY Blue;
 		
 		/** (Optional) CIE xyY of the White Point; this is used only if the ColorSpace is "Custom". */
-		FVector WhitePoint = FVector::ZeroVector;
+		FDMXGDTFColorCIE1931xyY WhitePoint;
 
 		/** The outer physcial descriptions */
 		const TWeakPtr<FDMXGDTFPhysicalDescriptions> OuterPhysicalDescriptions;
diff --git a/Engine/Plugins/VirtualProduction/DMX/DMXGDTF/Source/DMXGDTF/Internal/GDTF/PhysicalDescriptions/DMXGDTFEmitter.h b/Engine/Plugins/VirtualProduction/DMX/DMXGDTF/Source/DMXGDTF/Internal/GDTF/PhysicalDescriptions/DMXGDTFEmitter.h
--- a/Engine/Plugins/VirtualProduction/DMX/DMXGDTF/Source/DMXGDTF/Internal/GDTF/PhysicalDescriptions/DMXGDTFEmitter.h
+++ b/Engine/Plugins/VirtualProduction/DMX/DMXGDTF/Source/DMXGDTF/Internal/GDTF/PhysicalDescriptions/DMXGDTFEmitter.h
@@ -2,6 +2,7 @@
 
 #pragma once
 
+#include "DMXGDTFColorCIE1931xyY.h"
 #include "GDTF/DMXGDTFNode.h"
 #include "Math/Vector.h"
 #include "UObject/NameTypes.h"
@@ -21,6 +22,7 @@ namespace UE::DMX::GDTF
 		//~ Begin FDMXGDTFNode interface
 		virtual const TCHAR* GetXmlTag() const override { return TEXT("Emitter"); }
 		virtual void Initialize(const FXmlNode& XmlNode) override;
+		virtual FXmlNode* CreateXmlNode(FXmlNode& Parent) override;
 		//~ End FDMXGDTFNode interface
 		
 		/** Unique Name of the emitter */
@@ -30,7 +32,7 @@ namespace UE::DMX::GDTF
 		 * Approximate absolute color point if applicable.Omit for non - visible emitters (eg., UV).
 		 * For Y give relative value compared to overall output defined in property Luminous Flux of related Beam Geometry (transmissive case). 
 		 */
-		FVector ColorCIE;
+		FDMXGDTFColorCIE1931xyY Color;
 
 		/** Required if color is omitted, otherwise it is optional. Dominant wavelength of the LED. */
 		float DominantWaveLength;
diff --git a/Engine/Plugins/VirtualProduction/DMX/DMXGDTF/Source/DMXGDTF/Internal/GDTF/PhysicalDescriptions/DMXGDTFFilter.h b/Engine/Plugins/VirtualProduction/DMX/DMXGDTF/Source/DMXGDTF/Internal/GDTF/PhysicalDescriptions/DMXGDTFFilter.h
--- a/Engine/Plugins/VirtualProduction/DMX/DMXGDTF/Source/DMXGDTF/Internal/GDTF/PhysicalDescriptions/DMXGDTFFilter.h
+++ b/Engine/Plugins/VirtualProduction/DMX/DMXGDTF/Source/DMXGDTF/Internal/GDTF/PhysicalDescriptions/DMXGDTFFilter.h
@@ -2,6 +2,7 @@
 
 #pragma once
 
+#include "DMXGDTFColorCIE1931xyY.h"
 #include "GDTF/DMXGDTFNode.h"
 #include "Math/Vector.h"
 #include "UObject/NameTypes.h"
@@ -21,6 +22,7 @@ namespace UE::DMX::GDTF
 		//~ Begin FDMXGDTFNode interface
 		virtual const TCHAR* GetXmlTag() const override { return TEXT("Filter"); }
 		virtual void Initialize(const FXmlNode& XmlNode) override;
+		virtual FXmlNode* CreateXmlNode(FXmlNode& Parent) override;
 		//~ End FDMXGDTFNode interface
 
 		/** Unique Name of the filter */
@@ -30,7 +32,7 @@ namespace UE::DMX::GDTF
 		 * Approximate absolute color point when this filter is the only item fully inserted into the beam and the fixture is at maximum intensity.
 		 * For Y give relative value compared to overall output defined in property Luminous Flux of related Beam Geometry(transmissive case).
 		 */
-		FVector ColorCIE;
+		FDMXGDTFColorCIE1931xyY Color;
 
 		/** As children the Filter has a list of measurements. */
 		TArray<TSharedPtr<FDMXGDTFFilterMeasurement>> Measurements;
diff --git a/Engine/Plugins/VirtualProduction/DMX/DMXGDTF/Source/DMXGDTF/Internal/GDTF/PhysicalDescriptions/DMXGDTFGamut.h b/Engine/Plugins/VirtualProduction/DMX/DMXGDTF/Source/DMXGDTF/Internal/GDTF/PhysicalDescriptions/DMXGDTFGamut.h
--- a/Engine/Plugins/VirtualProduction/DMX/DMXGDTF/Source/DMXGDTF/Internal/GDTF/PhysicalDescriptions/DMXGDTFGamut.h
+++ b/Engine/Plugins/VirtualProduction/DMX/DMXGDTF/Source/DMXGDTF/Internal/GDTF/PhysicalDescriptions/DMXGDTFGamut.h
@@ -2,6 +2,7 @@
 
 #pragma once
 
+#include "DMXGDTFColorCIE1931xyY.h"
 #include "GDTF/DMXGDTFNode.h"
 #include "Math/Vector.h"
 #include "UObject/NameTypes.h"
@@ -20,19 +21,20 @@ namespace UE::DMX::GDTF
 		//~ Begin FDMXGDTFNode interface
 		virtual const TCHAR* GetXmlTag() const override { return TEXT("Gamut"); }
 		virtual void Initialize(const FXmlNode& XmlNode) override;
+		virtual FXmlNode* CreateXmlNode(FXmlNode& Parent) override;
 		//~ End FDMXGDTFNode interface
 
 		/** Unique Name of the Gamut. */
 		FName Name;	
 
 		/** Set of points defining the vertice of the gamut's polygon. */
-		TArray<FVector>	Points;
+		TArray<FDMXGDTFColorCIE1931xyY>	Points;
 
 		/** The outer physcial descriptions */
 		const TWeakPtr<FDMXGDTFPhysicalDescriptions> OuterPhysicalDescriptions;
 
 	private:
 		/** Parses points from a GDTF string */
-		TArray<FVector> ParsePoints(const FString& GDTFString) const;
+		TArray<FDMXGDTFColorCIE1931xyY> ParsePoints(const FString& GDTFString) const;
 	};
 }
diff --git a/Engine/Plugins/VirtualProduction/DMX/DMXGDTF/Source/DMXGDTF/Internal/GDTF/PhysicalDescriptions/DMXGDTFMeasurement.h b/Engine/Plugins/VirtualProduction/DMX/DMXGDTF/Source/DMXGDTF/Internal/GDTF/PhysicalDescriptions/DMXGDTFMeasurement.h
--- a/Engine/Plugins/VirtualProduction/DMX/DMXGDTF/Source/DMXGDTF/Internal/GDTF/PhysicalDescriptions/DMXGDTFMeasurement.h
+++ b/Engine/Plugins/VirtualProduction/DMX/DMXGDTF/Source/DMXGDTF/Internal/GDTF/PhysicalDescriptions/DMXGDTFMeasurement.h
@@ -48,19 +48,20 @@ namespace UE::DMX::GDTF
 		//~ Begin FDMXGDTFNode interface
 		virtual const TCHAR* GetXmlTag() const override { return TEXT("Measurement"); }
 		virtual void Initialize(const FXmlNode& XmlNode) override;
+		virtual FXmlNode* CreateXmlNode(FXmlNode& Parent) override;
 		//~ End FDMXGDTFNode interface
 		
 		/** 
 		 * For additive color mixing: uniquely given emitter intensity DMX percentage. Value range between > 0 and <= 100. 
 		 * For subtractive color mixing: uniquely given flag insertion DMX percentage. Value range between 0 and 100. 
 		 */
-		float Physical;		
+		float Physical = 0.f;		
 
 		/** Used for additive color mixing: overall candela value for the enclosed set of measurement. */
-		float LuminousIntensity;		
+		float LuminousIntensity = 0.f;		
 
 		/** Used for subtractive color mixing: total amount of lighting energy passed at this insertion percentage. */
-		float Transmission;	
+		float Transmission = 0.f;	
 
 		/** Interpolation scheme from the previous value.The currently defined values are : "Linear", "Step", "Log"; Default: Linear */
 		EDMXGDTFMeasurementInterpolationTo InterpolationTo = EDMXGDTFMeasurementInterpolationTo::Linear;
@@ -74,7 +75,7 @@ namespace UE::DMX::GDTF
 		: public FDMXGDTFMeasurementBase
 	{
 	public:
-		FDMXGDTFEmitterMeasurement(const TWeakPtr<FDMXGDTFEmitter>& InEmitter);
+		FDMXGDTFEmitterMeasurement(const TSharedRef<FDMXGDTFEmitter>& InEmitter);
 
 		/** The outer emitter */
 		const TWeakPtr<FDMXGDTFEmitter> OuterEmitter;
diff --git a/Engine/Plugins/VirtualProduction/DMX/DMXGDTF/Source/DMXGDTF/Internal/GDTF/PhysicalDescriptions/DMXGDTFPhysicalDescriptions.h b/Engine/Plugins/VirtualProduction/DMX/DMXGDTF/Source/DMXGDTF/Internal/GDTF/PhysicalDescriptions/DMXGDTFPhysicalDescriptions.h
--- a/Engine/Plugins/VirtualProduction/DMX/DMXGDTF/Source/DMXGDTF/Internal/GDTF/PhysicalDescriptions/DMXGDTFPhysicalDescriptions.h
+++ b/Engine/Plugins/VirtualProduction/DMX/DMXGDTF/Source/DMXGDTF/Internal/GDTF/PhysicalDescriptions/DMXGDTFPhysicalDescriptions.h
@@ -26,6 +26,7 @@ namespace UE::DMX::GDTF
 		//~ Begin FDMXGDTFNode interface
 		virtual const TCHAR* GetXmlTag() const override { return TEXT("PhysicalDescriptions"); }
 		virtual void Initialize(const FXmlNode& XmlNode) override;
+		virtual FXmlNode* CreateXmlNode(FXmlNode& Parent) override;
 		//~ End FDMXGDTFNode interface
 
 		/** 
@@ -43,7 +44,7 @@ namespace UE::DMX::GDTF
 		TArray<TSharedPtr<FDMXGDTFFilter>> Filters;
 
 		/** This section defines color spaces. Currently it does not have any XML attributes (XML node <AdditionalColorSpaces>). */
-		TArray<TSharedPtr<FDMXGDTFColorSpace>> ColorSpaces;
+		TSharedPtr<FDMXGDTFColorSpace> ColorSpaces;
 
 		/** Describes additional device color spaces. */
 		TArray<TSharedPtr<FDMXGDTFColorSpace>> AdditionalColorSpaces;
diff --git a/Engine/Plugins/VirtualProduction/DMX/DMXGDTF/Source/DMXGDTF/Internal/GDTF/PhysicalDescriptions/DMXGDTFProperties.h b/Engine/Plugins/VirtualProduction/DMX/DMXGDTF/Source/DMXGDTF/Internal/GDTF/PhysicalDescriptions/DMXGDTFProperties.h
--- a/Engine/Plugins/VirtualProduction/DMX/DMXGDTF/Source/DMXGDTF/Internal/GDTF/PhysicalDescriptions/DMXGDTFProperties.h
+++ b/Engine/Plugins/VirtualProduction/DMX/DMXGDTF/Source/DMXGDTF/Internal/GDTF/PhysicalDescriptions/DMXGDTFProperties.h
@@ -25,13 +25,14 @@ namespace UE::DMX::GDTF
 		//~ Begin FDMXGDTFNode interface
 		virtual const TCHAR* GetXmlTag() const override { return TEXT("Properties"); }
 		virtual void Initialize(const FXmlNode& XmlNode) override;
+		virtual FXmlNode* CreateXmlNode(FXmlNode& Parent) override;
 		//~ End FDMXGDTFNode interface
 
 		/** (optional) Temperature range in which the device can be operated. */
 		TSharedPtr<FDMXGDTFOperatingTemperature> OperatingTemperature;
 
 		/** (optional) Weight of the device including all accessories. */
-		TSharedPtr<FDMXGDTFWeight> Weigth;
+		TSharedPtr<FDMXGDTFWeight> Weight;
 
 		/** (optional) Height of the legs. */
 		TSharedPtr<FDMXGDTFLegHeight> LegHeight;
diff --git a/Engine/Plugins/VirtualProduction/DMX/DMXGDTF/Source/DMXGDTF/Internal/GDTF/Protocols/DMXGDTFProtocolDMXMap.h b/Engine/Plugins/VirtualProduction/DMX/DMXGDTF/Source/DMXGDTF/Internal/GDTF/Protocols/DMXGDTFProtocolDMXMap.h
--- a/Engine/Plugins/VirtualProduction/DMX/DMXGDTF/Source/DMXGDTF/Internal/GDTF/Protocols/DMXGDTFProtocolDMXMap.h
+++ b/Engine/Plugins/VirtualProduction/DMX/DMXGDTF/Source/DMXGDTF/Internal/GDTF/Protocols/DMXGDTFProtocolDMXMap.h
@@ -17,6 +17,7 @@ namespace UE::DMX::GDTF
 		//~ Begin FDMXGDTFNode interface
 		virtual const TCHAR* GetXmlTag() const override { return TEXT("Protocols"); }
 		virtual void Initialize(const FXmlNode& XmlNode) override;
+		virtual FXmlNode* CreateXmlNode(FXmlNode& Parent) override;
 		//~ End FDMXGDTFNode interface
 
 		/** Value of the protocol value. */
@@ -34,7 +35,7 @@ namespace UE::DMX::GDTF
 			: public FDMXGDTFProtocolDMXMapBase
 		{
 		public:
-			FDMXGDTFProtocolArtNetDMXMap(const TWeakPtr<ArtNet::FDMXGDTFProtocolArtNet>& InProtocolArtNet);
+			FDMXGDTFProtocolArtNetDMXMap(const TSharedRef<ArtNet::FDMXGDTFProtocolArtNet>& InProtocolArtNet);
 
 			const TWeakPtr<FDMXGDTFProtocolArtNet> OuterProtocolArtNet;
 		};
diff --git a/Engine/Plugins/VirtualProduction/DMX/DMXGDTF/Source/DMXGDTF/Internal/GDTF/Protocols/DMXGDTFProtocols.h b/Engine/Plugins/VirtualProduction/DMX/DMXGDTF/Source/DMXGDTF/Internal/GDTF/Protocols/DMXGDTFProtocols.h
--- a/Engine/Plugins/VirtualProduction/DMX/DMXGDTF/Source/DMXGDTF/Internal/GDTF/Protocols/DMXGDTFProtocols.h
+++ b/Engine/Plugins/VirtualProduction/DMX/DMXGDTF/Source/DMXGDTF/Internal/GDTF/Protocols/DMXGDTFProtocols.h
@@ -7,9 +7,9 @@
 namespace UE::DMX::GDTF
 {
 	class FDMXGDTFFixtureType;
-	namespace RDM { class FDMXGDTFProtocolRDM; }
+	namespace RDM { class FDMXGDTFProtocolFTRDM; }
 	namespace ArtNet { class FDMXGDTFProtocolArtNet; }
-	namespace SACN { class FDMXGDTFProtocolsACN; }
+	namespace SACN { class FDMXGDTFProtocolSACN; }
 
 	/** This section defines the overall Protocols of the device (XML node <Protocols>). */
 	class DMXGDTF_API FDMXGDTFProtocols
@@ -21,16 +21,17 @@ namespace UE::DMX::GDTF
 		//~ Begin FDMXGDTFNode interface
 		virtual const TCHAR* GetXmlTag() const override { return TEXT("Protocols"); }
 		virtual void Initialize(const FXmlNode& XmlNode) override;
+		virtual FXmlNode* CreateXmlNode(FXmlNode& Parent) override;
 		//~ End FDMXGDTFNode interface
 
 		/** Describes RDM information */
-		TSharedPtr<RDM::FDMXGDTFProtocolRDM> RDM;
+		TSharedPtr<RDM::FDMXGDTFProtocolFTRDM> RDM;
 
 		/** Describes Art-Net information */
 		TSharedPtr<ArtNet::FDMXGDTFProtocolArtNet> ArtNet;
 
 		/** Describes sACN information */
-		TSharedPtr<SACN::FDMXGDTFProtocolsACN> sACN;
+		TSharedPtr<SACN::FDMXGDTFProtocolSACN> sACN;
 
 		/** The outer fixture type */
 		const TWeakPtr<FDMXGDTFFixtureType> OuterFixtureType;
diff --git a/Engine/Plugins/VirtualProduction/DMX/DMXGDTF/Source/DMXGDTF/Internal/GDTF/Protocols/RDM/DMXGDTFSoftwareVersionID.h b/Engine/Plugins/VirtualProduction/DMX/DMXGDTF/Source/DMXGDTF/Internal/GDTF/Protocols/RDM/DMXGDTFSoftwareVersionID.h
--- a/Engine/Plugins/VirtualProduction/DMX/DMXGDTF/Source/DMXGDTF/Internal/GDTF/Protocols/RDM/DMXGDTFSoftwareVersionID.h
+++ b/Engine/Plugins/VirtualProduction/DMX/DMXGDTF/Source/DMXGDTF/Internal/GDTF/Protocols/RDM/DMXGDTFSoftwareVersionID.h
@@ -7,18 +7,19 @@
 namespace UE::DMX::GDTF::RDM
 {
 	class FDMXGDTFDMXPersonality;
-	class FDMXGDTFProtocolRDM;
+	class FDMXGDTFProtocolFTRDM;
 
 	/** For each supported software version add an XML node <SoftwareVersionID>.. */
 	class DMXGDTF_API FDMXGDTFSoftwareVersionID
 		: public FDMXGDTFNode
 	{
 	public:
-		FDMXGDTFSoftwareVersionID(const TSharedRef<FDMXGDTFProtocolRDM>& InProtocolRDM);
+		FDMXGDTFSoftwareVersionID(const TSharedRef<FDMXGDTFProtocolFTRDM>& InProtocolRDM);
 
 		//~ Begin FDMXGDTFNode interface
 		virtual const TCHAR* GetXmlTag() const override { return TEXT("SoftwareVersionID"); }
 		virtual void Initialize(const FXmlNode& XmlNode) override;
+		virtual FXmlNode* CreateXmlNode(FXmlNode& Parent) override;
 		//~ End FDMXGDTFNode interface
 
 		/** Software version ID */
@@ -28,6 +29,6 @@ namespace UE::DMX::GDTF::RDM
 		TArray<TSharedPtr<FDMXGDTFDMXPersonality>> DMXPersonalityArray;
 
 		/** The outer protocols */
-		const TWeakPtr<FDMXGDTFProtocolRDM> OuterProtocolRDM;
+		const TWeakPtr<FDMXGDTFProtocolFTRDM> OuterProtocolRDM;
 	};
 }
diff --git a/Engine/Plugins/VirtualProduction/DMX/DMXGDTF/Source/DMXGDTF/Internal/GDTF/Protocols/sACN/DMXGDTFProtocolSACN.h b/Engine/Plugins/VirtualProduction/DMX/DMXGDTF/Source/DMXGDTF/Internal/GDTF/Protocols/sACN/DMXGDTFProtocolSACN.h
--- a/Engine/Plugins/VirtualProduction/DMX/DMXGDTF/Source/DMXGDTF/Internal/GDTF/Protocols/sACN/DMXGDTFProtocolSACN.h
+++ b/Engine/Plugins/VirtualProduction/DMX/DMXGDTF/Source/DMXGDTF/Internal/GDTF/Protocols/sACN/DMXGDTFProtocolSACN.h
@@ -20,8 +20,9 @@ namespace UE::DMX::GDTF
 			FDMXGDTFProtocolSACN(const TSharedRef<FDMXGDTFProtocols>& InProtocols);
 
 			//~ Begin FDMXGDTFNode interface
-			virtual const TCHAR* GetXmlTag() const override { return TEXT("FTRDM"); }
+			virtual const TCHAR* GetXmlTag() const override { return TEXT("sACN"); }
 			virtual void Initialize(const FXmlNode& XmlNode) override;
+			virtual FXmlNode* CreateXmlNode(FXmlNode& Parent) override;
 			//~ End FDMXGDTFNode interface
 
 			/** As children the Art-Net has a list of Maps. */
diff --git a/Engine/Plugins/VirtualProduction/DMX/DMXGDTF/Source/DMXGDTF/Internal/GDTF/Revisions/DMXGDTFRevision.h b/Engine/Plugins/VirtualProduction/DMX/DMXGDTF/Source/DMXGDTF/Internal/GDTF/Revisions/DMXGDTFRevision.h
--- a/Engine/Plugins/VirtualProduction/DMX/DMXGDTF/Source/DMXGDTF/Internal/GDTF/Revisions/DMXGDTFRevision.h
+++ b/Engine/Plugins/VirtualProduction/DMX/DMXGDTF/Source/DMXGDTF/Internal/GDTF/Revisions/DMXGDTFRevision.h
@@ -23,6 +23,7 @@ namespace UE::DMX::GDTF
 		//~ Begin FDMXGDTFNode interface
 		virtual const TCHAR* GetXmlTag() const override { return TEXT("Revision"); }
 		virtual void Initialize(const FXmlNode& XmlNode) override;
+		virtual FXmlNode* CreateXmlNode(FXmlNode& Parent) override;
 		//~ End FDMXGDTFNode interface;
 
 		/** User-defined text for this revision; Default value: empty */
@@ -32,7 +33,7 @@ namespace UE::DMX::GDTF
 		FDateTime Date;
 
 		/** UserID of the user that has uploaded the GDTF file to the database; Default value: 0 */
-		uint32 User = 0;
+		uint32 UserID = 0;
 
 		/** Name of the software that modified this revision; Default value: empty */
 		FString ModifiedBy;
diff --git a/Engine/Plugins/VirtualProduction/DMX/DMXGDTF/Source/DMXGDTF/Internal/GDTF/Wheels/DMXGDTFAnimationSystem.h b/Engine/Plugins/VirtualProduction/DMX/DMXGDTF/Source/DMXGDTF/Internal/GDTF/Wheels/DMXGDTFAnimationSystem.h
--- a/Engine/Plugins/VirtualProduction/DMX/DMXGDTF/Source/DMXGDTF/Internal/GDTF/Wheels/DMXGDTFAnimationSystem.h
+++ b/Engine/Plugins/VirtualProduction/DMX/DMXGDTF/Source/DMXGDTF/Internal/GDTF/Wheels/DMXGDTFAnimationSystem.h
@@ -19,6 +19,7 @@ namespace UE::DMX::GDTF
 		//~ Begin FDMXGDTFNode interface
 		virtual const TCHAR* GetXmlTag() const override { return TEXT("AnimationSystem"); }
 		virtual void Initialize(const FXmlNode& XmlNode) override;
+		virtual FXmlNode* CreateXmlNode(FXmlNode& Parent) override;
 		//~ End FDMXGDTFNode interface
 
 		/**
@@ -52,7 +53,7 @@ namespace UE::DMX::GDTF
 		 * Radius of the circle that defines the section of the animation system which 
 		 * will be shown in the beam 
 		 */
-		double Radius = 0.0;
+		float Radius = 0.0;
 
 		/** The outer wheel slot */
 		const TWeakPtr<FDMXGDTFWheelSlot> OuterWheelSlot;
diff --git a/Engine/Plugins/VirtualProduction/DMX/DMXGDTF/Source/DMXGDTF/Internal/GDTF/Wheels/DMXGDTFPrismFacet.h b/Engine/Plugins/VirtualProduction/DMX/DMXGDTF/Source/DMXGDTF/Internal/GDTF/Wheels/DMXGDTFPrismFacet.h
--- a/Engine/Plugins/VirtualProduction/DMX/DMXGDTF/Source/DMXGDTF/Internal/GDTF/Wheels/DMXGDTFPrismFacet.h
+++ b/Engine/Plugins/VirtualProduction/DMX/DMXGDTF/Source/DMXGDTF/Internal/GDTF/Wheels/DMXGDTFPrismFacet.h
@@ -2,6 +2,7 @@
 
 #pragma once
 
+#include "DMXGDTFColorCIE1931xyY.h"
 #include "GDTF/DMXGDTFNode.h"
 #include "Math/Vector.h"
 #include "Math/Transform.h" 
@@ -20,6 +21,7 @@ namespace UE::DMX::GDTF
 		//~ Begin FDMXGDTFNode interface
 		virtual const TCHAR* GetXmlTag() const override { return TEXT("Facet"); }
 		virtual void Initialize(const FXmlNode& XmlNode) override;
+		virtual FXmlNode* CreateXmlNode(FXmlNode& Parent) override;
 		//~ End FDMXGDTFNode interface
 
 		/**
@@ -27,7 +29,7 @@ namespace UE::DMX::GDTF
 		 *
 		 * UE specific: Using a Vector3d to store the xyY color.
 		 */
-		FVector Color = { 0.3127, 0.3290, 100.0 };
+		FDMXGDTFColorCIE1931xyY Color = { 0.3127, 0.3290, 100.0 };
 
 		/** Specify the rotation, translation and scaling for the facet. */
 		FTransform Rotation = FTransform::Identity;
diff --git a/Engine/Plugins/VirtualProduction/DMX/DMXGDTF/Source/DMXGDTF/Internal/GDTF/Wheels/DMXGDTFWheelSlot.h b/Engine/Plugins/VirtualProduction/DMX/DMXGDTF/Source/DMXGDTF/Internal/GDTF/Wheels/DMXGDTFWheelSlot.h
--- a/Engine/Plugins/VirtualProduction/DMX/DMXGDTF/Source/DMXGDTF/Internal/GDTF/Wheels/DMXGDTFWheelSlot.h
+++ b/Engine/Plugins/VirtualProduction/DMX/DMXGDTF/Source/DMXGDTF/Internal/GDTF/Wheels/DMXGDTFWheelSlot.h
@@ -2,6 +2,7 @@
 
 #pragma once
 
+#include "DMXGDTFColorCIE1931xyY.h"
 #include "GDTF/DMXGDTFNode.h"
 #include "Math/Vector.h"
 
@@ -21,16 +22,18 @@ namespace UE::DMX::GDTF
 		//~ Begin FDMXGDTFNode interface
 		virtual const TCHAR* GetXmlTag() const override { return TEXT("Slot"); }
 		virtual void Initialize(const FXmlNode& XmlNode) override;
+		virtual FXmlNode* CreateXmlNode(FXmlNode& Parent) override;
 		//~ End FDMXGDTFNode interface
 
+		/** The unique name of the wheel slot */
+		FName Name;
+
 		/**
 		 * Color of the wheel slot, Default value: {0.3127, 0.3290, 100.0} (white). 
 		 * For Y give relative value compared to overall output defined in property
 		 * Luminous Flux of related Beam Geometry (transmissive case).
-		 *
-		 * UE specific: Using a Vector3d to store the xyY color.
 		 */
-		FVector Color = { 0.3127, 0.3290, 100.0 };
+		FDMXGDTFColorCIE1931xyY Color;
 
 		/**
 		 * (Optional) PNG file name without extension containing image for specific gobos etc.
diff --git a/Engine/Plugins/VirtualProduction/DMX/DMXGDTF/Source/DMXGDTF/Private/DMXGDTF.cpp b/Engine/Plugins/VirtualProduction/DMX/DMXGDTF/Source/DMXGDTF/Private/DMXGDTF.cpp
--- a/Engine/Plugins/VirtualProduction/DMX/DMXGDTF/Source/DMXGDTF/Private/DMXGDTF.cpp
+++ b/Engine/Plugins/VirtualProduction/DMX/DMXGDTF/Source/DMXGDTF/Private/DMXGDTF.cpp
@@ -35,3 +35,14 @@ void UDMXGDTF::InitializeFromData(const TArray64<uint8>& Data)
 	Description = MakeShared<FDMXGDTFDescription>();
 	Description->InitializeFromDescriptionXml(XmlFile);
 }
+
+void UDMXGDTF::InitializeFromFixtureType(const TSharedRef<FDMXGDTFFixtureType>& InFixtureType)
+{
+	Description = MakeShared<FDMXGDTFDescription>();
+	Description->InitializeFromFixtureType(InFixtureType);
+}
+
+TSharedPtr<FXmlFile> UDMXGDTF::ExportAsXml() const
+{
+	return Description->ExportAsXml();
+}
diff --git a/Engine/Plugins/VirtualProduction/DMX/DMXGDTF/Source/DMXGDTF/Private/GDTF/AttributeDefinitions/DMXGDTFActivationGroup.cpp b/Engine/Plugins/VirtualProduction/DMX/DMXGDTF/Source/DMXGDTF/Private/GDTF/AttributeDefinitions/DMXGDTFActivationGroup.cpp
--- a/Engine/Plugins/VirtualProduction/DMX/DMXGDTF/Source/DMXGDTF/Private/GDTF/AttributeDefinitions/DMXGDTFActivationGroup.cpp
+++ b/Engine/Plugins/VirtualProduction/DMX/DMXGDTF/Source/DMXGDTF/Private/GDTF/AttributeDefinitions/DMXGDTFActivationGroup.cpp
@@ -3,6 +3,7 @@
 #include "GDTF/AttributeDefinitions/DMXGDTFActivationGroup.h"
 
 #include "Serialization/DMXGDTFNodeInitializer.h"
+#include "Serialization/DMXGDTFXmlNodeBuilder.h"
 
 namespace UE::DMX::GDTF
 {
@@ -15,4 +16,12 @@ namespace UE::DMX::GDTF
 		FDMXGDTFNodeInitializer(SharedThis(this), XmlNode)
 			.GetAttribute(TEXT("Name"), Name);
 	}
+
+	FXmlNode* FDMXGDTFActivationGroup::CreateXmlNode(FXmlNode& Parent)
+	{
+		const FDMXGDTFXmlNodeBuilder ChildBuilder = FDMXGDTFXmlNodeBuilder(Parent, *this)
+			.SetAttribute(TEXT("Name"), Name);
+
+		return ChildBuilder.GetIntermediateXmlNode();
+	}
 }
diff --git a/Engine/Plugins/VirtualProduction/DMX/DMXGDTF/Source/DMXGDTF/Private/GDTF/AttributeDefinitions/DMXGDTFAttribute.cpp b/Engine/Plugins/VirtualProduction/DMX/DMXGDTF/Source/DMXGDTF/Private/GDTF/AttributeDefinitions/DMXGDTFAttribute.cpp
--- a/Engine/Plugins/VirtualProduction/DMX/DMXGDTF/Source/DMXGDTF/Private/GDTF/AttributeDefinitions/DMXGDTFAttribute.cpp
+++ b/Engine/Plugins/VirtualProduction/DMX/DMXGDTF/Source/DMXGDTF/Private/GDTF/AttributeDefinitions/DMXGDTFAttribute.cpp
@@ -3,8 +3,10 @@
 #include "GDTF/AttributeDefinitions/DMXGDTFAttribute.h"
 
 #include "GDTF/AttributeDefinitions/DMXGDTFAttributeDefinitions.h"
+#include "GDTF/AttributeDefinitions/DMXGDTFPhysicalUnit.h"
 #include "GDTF/AttributeDefinitions/DMXGDTFSubphysicalUnit.h"
 #include "Serialization/DMXGDTFNodeInitializer.h"
+#include "Serialization/DMXGDTFXmlNodeBuilder.h"
 
 namespace UE::DMX::GDTF
 {
@@ -25,6 +27,25 @@ namespace UE::DMX::GDTF
 			.CreateChildren(TEXT("SubphysicalUnit"), SubpyhsicalUnitArray);
 	}
 
+	FXmlNode* FDMXGDTFAttribute::CreateXmlNode(FXmlNode& Parent)
+	{
+		const FString DefaultActivationGroup = TEXT("");
+		const FString DefaultMainAttribute = TEXT("");
+		const FDMXGDTFColorCIE1931xyY DefaultColor = { 0.f, 0.f, 0.f };
+
+		const FDMXGDTFXmlNodeBuilder ChildBuilder = FDMXGDTFXmlNodeBuilder(Parent, *this)
+			.SetAttribute(TEXT("Name"), Name)
+			.SetAttribute(TEXT("Pretty"), Pretty)
+			.SetAttribute(TEXT("PhysicalUnit"), PhysicalUnit)
+			.SetAttribute(TEXT("ActivationGroup"), ActivationGroup, DefaultActivationGroup)
+			.SetAttribute(TEXT("Feature"), Feature)
+			.SetAttribute(TEXT("MainAttribute"), MainAttribute, DefaultMainAttribute)
+			.SetAttribute(TEXT("Color"), Color, DefaultColor)
+			.AppendChildren(TEXT("SubphysicalUnit"), SubpyhsicalUnitArray);
+
+		return ChildBuilder.GetIntermediateXmlNode();
+	}
+
 	TSharedPtr<FDMXGDTFActivationGroup> FDMXGDTFAttribute::ResolveActivationGroup() const
 	{
 		if (const TSharedPtr<FDMXGDTFAttributeDefinitions> AttributeDefinitions = OuterAttributeDefinitions.Pin())
diff --git a/Engine/Plugins/VirtualProduction/DMX/DMXGDTF/Source/DMXGDTF/Private/GDTF/AttributeDefinitions/DMXGDTFAttributeDefinitions.cpp b/Engine/Plugins/VirtualProduction/DMX/DMXGDTF/Source/DMXGDTF/Private/GDTF/AttributeDefinitions/DMXGDTFAttributeDefinitions.cpp
--- a/Engine/Plugins/VirtualProduction/DMX/DMXGDTF/Source/DMXGDTF/Private/GDTF/AttributeDefinitions/DMXGDTFAttributeDefinitions.cpp
+++ b/Engine/Plugins/VirtualProduction/DMX/DMXGDTF/Source/DMXGDTF/Private/GDTF/AttributeDefinitions/DMXGDTFAttributeDefinitions.cpp
@@ -8,14 +8,12 @@
 #include "GDTF/AttributeDefinitions/DMXGDTFFeature.h"
 #include "GDTF/AttributeDefinitions/DMXGDTFFeatureGroup.h"
 #include "Serialization/DMXGDTFNodeInitializer.h"
+#include "Serialization/DMXGDTFXmlNodeBuilder.h"
 
 namespace UE::DMX::GDTF
 {
 	FDMXGDTFAttributeDefinitions::FDMXGDTFAttributeDefinitions(const TSharedRef<FDMXGDTFFixtureType>& InFixtureType)
-	{		
-		// The outer fixture type is available to all nodes and doesn't need extra initialization.
-		// This constructor is only here to enforce the correct outer.
-	}
+	{}
 
 	void FDMXGDTFAttributeDefinitions::Initialize(const FXmlNode& XmlNode)
 	{
@@ -25,6 +23,16 @@ namespace UE::DMX::GDTF
 			.CreateChildCollection(TEXT("Attributes"), TEXT("Attribute"), Attributes);
 	}
 
+	FXmlNode* FDMXGDTFAttributeDefinitions::CreateXmlNode(FXmlNode& Parent)
+	{
+		const FDMXGDTFXmlNodeBuilder ChildBuilder = FDMXGDTFXmlNodeBuilder(Parent, *this)
+			.AppendChildCollection(TEXT("ActivationGroups"), TEXT("ActivationGroup"), ActivationGroups)
+			.AppendChildCollection(TEXT("FeatureGroups"), TEXT("FeatureGroup"), FeatureGroups)
+			.AppendChildCollection(TEXT("Attributes"), TEXT("Attribute"), Attributes);
+
+		return ChildBuilder.GetIntermediateXmlNode();
+	}
+
 	TSharedPtr<FDMXGDTFActivationGroup> FDMXGDTFAttributeDefinitions::FindActivationGroup(const FString& ActivationGroupName) const
 	{
 		if (const TSharedPtr<FDMXGDTFActivationGroup>* ActivationGroupPtr = Algo::FindBy(ActivationGroups, ActivationGroupName, &FDMXGDTFActivationGroup::Name))
diff --git a/Engine/Plugins/VirtualProduction/DMX/DMXGDTF/Source/DMXGDTF/Private/GDTF/AttributeDefinitions/DMXGDTFFeature.cpp b/Engine/Plugins/VirtualProduction/DMX/DMXGDTF/Source/DMXGDTF/Private/GDTF/AttributeDefinitions/DMXGDTFFeature.cpp
--- a/Engine/Plugins/VirtualProduction/DMX/DMXGDTF/Source/DMXGDTF/Private/GDTF/AttributeDefinitions/DMXGDTFFeature.cpp
+++ b/Engine/Plugins/VirtualProduction/DMX/DMXGDTF/Source/DMXGDTF/Private/GDTF/AttributeDefinitions/DMXGDTFFeature.cpp
@@ -3,10 +3,11 @@
 #include "GDTF/AttributeDefinitions/DMXGDTFFeature.h"
 
 #include "Serialization/DMXGDTFNodeInitializer.h"
+#include "Serialization/DMXGDTFXmlNodeBuilder.h"
 
 namespace UE::DMX::GDTF
 {
-	FDMXGDTFFeature::FDMXGDTFFeature(const TWeakPtr<FDMXGDTFFeatureGroup>& InFeatureGroup)
+	FDMXGDTFFeature::FDMXGDTFFeature(const TSharedRef<FDMXGDTFFeatureGroup>& InFeatureGroup)
 		: OuterFeatureGroup(InFeatureGroup)
 	{}
 
@@ -15,4 +16,12 @@ namespace UE::DMX::GDTF
 		FDMXGDTFNodeInitializer(SharedThis(this), XmlNode)
 			.GetAttribute(TEXT("Name"), Name);
 	}
+	
+	FXmlNode* FDMXGDTFFeature::CreateXmlNode(FXmlNode& Parent)
+	{
+		const FDMXGDTFXmlNodeBuilder ChildBuilder = FDMXGDTFXmlNodeBuilder(Parent, *this)
+			.SetAttribute(TEXT("Name"), Name);
+
+		return ChildBuilder.GetIntermediateXmlNode();
+	}
 }
diff --git a/Engine/Plugins/VirtualProduction/DMX/DMXGDTF/Source/DMXGDTF/Private/GDTF/AttributeDefinitions/DMXGDTFFeatureGroup.cpp b/Engine/Plugins/VirtualProduction/DMX/DMXGDTF/Source/DMXGDTF/Private/GDTF/AttributeDefinitions/DMXGDTFFeatureGroup.cpp
--- a/Engine/Plugins/VirtualProduction/DMX/DMXGDTF/Source/DMXGDTF/Private/GDTF/AttributeDefinitions/DMXGDTFFeatureGroup.cpp
+++ b/Engine/Plugins/VirtualProduction/DMX/DMXGDTF/Source/DMXGDTF/Private/GDTF/AttributeDefinitions/DMXGDTFFeatureGroup.cpp
@@ -4,6 +4,7 @@
 
 #include "GDTF/AttributeDefinitions/DMXGDTFFeature.h"
 #include "Serialization/DMXGDTFNodeInitializer.h"
+#include "Serialization/DMXGDTFXmlNodeBuilder.h"
 
 namespace UE::DMX::GDTF
 {
@@ -18,4 +19,14 @@ namespace UE::DMX::GDTF
 			.GetAttribute(TEXT("Pretty"), Pretty)
 			.CreateChildren(TEXT("Feature"), FeatureArray);
 	}
+
+	FXmlNode* FDMXGDTFFeatureGroup::CreateXmlNode(FXmlNode& Parent)
+	{
+		const FDMXGDTFXmlNodeBuilder ChildBuilder = FDMXGDTFXmlNodeBuilder(Parent, *this)
+			.SetAttribute(TEXT("Name"), Name)
+			.SetAttribute(TEXT("Pretty"), Pretty)
+			.AppendChildren(TEXT("Feature"), FeatureArray);
+
+		return ChildBuilder.GetIntermediateXmlNode();
+	}
 }
diff --git a/Engine/Plugins/VirtualProduction/DMX/DMXGDTF/Source/DMXGDTF/Private/GDTF/AttributeDefinitions/DMXGDTFSubphysicalUnit.cpp b/Engine/Plugins/VirtualProduction/DMX/DMXGDTF/Source/DMXGDTF/Private/GDTF/AttributeDefinitions/DMXGDTFSubphysicalUnit.cpp
--- a/Engine/Plugins/VirtualProduction/DMX/DMXGDTF/Source/DMXGDTF/Private/GDTF/AttributeDefinitions/DMXGDTFSubphysicalUnit.cpp
+++ b/Engine/Plugins/VirtualProduction/DMX/DMXGDTF/Source/DMXGDTF/Private/GDTF/AttributeDefinitions/DMXGDTFSubphysicalUnit.cpp
@@ -4,6 +4,7 @@
 
 #include "GDTF/AttributeDefinitions/DMXGDTFPhysicalUnit.h"
 #include "Serialization/DMXGDTFNodeInitializer.h"
+#include "Serialization/DMXGDTFXmlNodeBuilder.h"
 
 namespace UE::DMX::GDTF
 {
@@ -19,4 +20,15 @@ namespace UE::DMX::GDTF
 			.GetAttribute(TEXT("PhysicalFrom"), PhysicalFrom)
 			.GetAttribute(TEXT("PhysicalTo"), PhysicalTo);
 	}
+
+	FXmlNode* FDMXGDTFSubphysicalUnit::CreateXmlNode(FXmlNode& Parent)
+	{
+		const FDMXGDTFXmlNodeBuilder ChildBuilder = FDMXGDTFXmlNodeBuilder(Parent, *this)
+			.SetAttribute(TEXT("Type"), Type)
+			.SetAttribute(TEXT("PhysicalUnit"), PhysicalUnit)
+			.SetAttribute(TEXT("PhysicalFrom"), PhysicalFrom)
+			.SetAttribute(TEXT("PhysicalTo"), PhysicalTo);
+
+		return ChildBuilder.GetIntermediateXmlNode();
+	}
 }
diff --git a/Engine/Plugins/VirtualProduction/DMX/DMXGDTF/Source/DMXGDTF/Private/GDTF/DMXGDTFDescription.cpp b/Engine/Plugins/VirtualProduction/DMX/DMXGDTF/Source/DMXGDTF/Private/GDTF/DMXGDTFDescription.cpp
--- a/Engine/Plugins/VirtualProduction/DMX/DMXGDTF/Source/DMXGDTF/Private/GDTF/DMXGDTFDescription.cpp
+++ b/Engine/Plugins/VirtualProduction/DMX/DMXGDTF/Source/DMXGDTF/Private/GDTF/DMXGDTFDescription.cpp
@@ -4,7 +4,6 @@
 
 #include "DMXGDTFLog.h"
 #include "GDTF/DMXGDTFFixtureType.h"
-#include "GDTF/DMXGDTFVersion.h"
 #include "Misc/MessageDialog.h"
 #include "XmlFile.h"
 #include "XmlNode.h"
@@ -68,4 +67,44 @@ namespace UE::DMX::GDTF
 			FixtureType->Initialize(*FixtureTypeXmlNode);
 		}
 	}
+
+	void FDMXGDTFDescription::InitializeFromFixtureType(const TSharedRef<FDMXGDTFFixtureType>& InFixtureType)
+	{	
+		FixtureType = InFixtureType;
+	}
+
+	TSharedPtr<FXmlFile> FDMXGDTFDescription::ExportAsXml() const
+	{
+		// Don't export if there's nothing to export
+		if (!FixtureType.IsValid())
+		{
+			return nullptr;
+		}
+
+		// Create Xml File
+		const FString Buffer = "<?xml version=\"1.0\" encoding=\"UTF - 8\" standalone=\"no\" ?>\n<GDTF>\n</GDTF>";
+
+		TSharedRef<FXmlFile> XmlFile = MakeShared<FXmlFile>();
+		const bool bCreatedNewFile = XmlFile->LoadFile(Buffer, EConstructMethod::ConstructFromBuffer);
+
+		FXmlNode* RootNode = XmlFile->GetRootNode();
+		if (!ensureAlwaysMsgf(bCreatedNewFile && RootNode, TEXT("Failed to create a GDTF description.xml. Cannot export GDTF.")))
+		{
+			return nullptr;
+		}
+		check(RootNode);
+
+		// Version the Root Node
+		constexpr TCHAR DataVersionAttributeName[] = TEXT("DataVersion");
+		TArray<FXmlAttribute> Attributes =
+		{
+			FXmlAttribute(DataVersionAttributeName, FDMXGDTFVersion::GetAsString()),
+		};
+		RootNode->SetAttributes(Attributes);
+
+		// Export Children
+		FixtureType->CreateXmlNode(*RootNode);
+
+		return XmlFile;
+	}
 }
diff --git a/Engine/Plugins/VirtualProduction/DMX/DMXGDTF/Source/DMXGDTF/Private/GDTF/DMXGDTFFixtureType.cpp b/Engine/Plugins/VirtualProduction/DMX/DMXGDTF/Source/DMXGDTF/Private/GDTF/DMXGDTFFixtureType.cpp
--- a/Engine/Plugins/VirtualProduction/DMX/DMXGDTF/Source/DMXGDTF/Private/GDTF/DMXGDTFFixtureType.cpp
+++ b/Engine/Plugins/VirtualProduction/DMX/DMXGDTF/Source/DMXGDTF/Private/GDTF/DMXGDTFFixtureType.cpp
@@ -4,6 +4,7 @@
 
 #include "GDTF/AttributeDefinitions/DMXGDTFAttributeDefinitions.h"
 #include "GDTF/DMXModes/DMXGDTFDMXMode.h"
+#include "GDTF/FTPresets/DMXGDTFFTPreset.h"
 #include "GDTF/Geometries/DMXGDTFGeometryCollect.h"
 #include "GDTF/Models/DMXGDTFModel.h"
 #include "GDTF/PhysicalDescriptions/DMXGDTFPhysicalDescriptions.h"
@@ -11,6 +12,7 @@
 #include "GDTF/Revisions/DMXGDTFRevision.h"
 #include "GDTF/Wheels/DMXGDTFWheel.h"
 #include "Serialization/DMXGDTFNodeInitializer.h"
+#include "Serialization/DMXGDTFXmlNodeBuilder.h"
 
 namespace UE::DMX::GDTF
 {
@@ -40,6 +42,37 @@ namespace UE::DMX::GDTF
 			.CreateOptionalChild(TEXT("Geometries"), GeometryCollect)
 			.CreateChildCollection(TEXT("DMXModes"), TEXT("DMXMode"), DMXModes)
 			.CreateChildCollection(TEXT("Revisions"), TEXT("Revision"), Revisions)
+			.CreateChildCollection(TEXT("FTPresets"), TEXT("FTPreset"), FTPresets)
 			.CreateOptionalChild(TEXT("Protocols"), Protocols);
 	}
+
+	FXmlNode* FDMXGDTFFixtureType::CreateXmlNode(FXmlNode& Parent)
+	{
+		const int32 DefaultThumbnailOffsetX = 0;
+		const int32 DefaultThumbnailOffsetY = 0;
+
+		const FDMXGDTFXmlNodeBuilder ChildBuilder = FDMXGDTFXmlNodeBuilder(Parent, *this)
+			.SetAttribute(TEXT("Name"), Name)
+			.SetAttribute(TEXT("ShortName"), ShortName)
+			.SetAttribute(TEXT("LongName"), LongName)
+			.SetAttribute(TEXT("Manufacturer"), Manufacturer)
+			.SetAttribute(TEXT("Description"), Description)
+			.SetAttribute(TEXT("FixtureTypeID"), FixtureTypeID)
+			.SetAttribute(TEXT("Thumbnail"), Thumbnail)
+			.SetAttribute(TEXT("ThumbnailOffsetX"), ThumbnailOffsetX, DefaultThumbnailOffsetX)
+			.SetAttribute(TEXT("ThumbnailOffsetY"), ThumbnailOffsetY, DefaultThumbnailOffsetY)
+			.SetAttribute(TEXT("RefFT"), RefFT)
+			.SetAttribute(TEXT("CanHaveChildren"), bCanHaveChildren ? TEXT("Yes") : TEXT("No"))
+			.AppendRequiredChild(TEXT("AttributeDefinitions"), AttributeDefinitions)
+			.AppendChildCollection(TEXT("Wheels"), TEXT("Wheel"), Wheels)
+			.AppendOptionalChild(TEXT("PhysicalDescriptions"), PhysicalDescriptions)
+			.AppendChildCollection(TEXT("Models"), TEXT("Model"), Models)
+			.AppendOptionalChild(TEXT("Geometries"), GeometryCollect)
+			.AppendChildCollection(TEXT("DMXModes"), TEXT("DMXMode"), DMXModes)
+			.AppendChildCollection(TEXT("FTPresets"), TEXT("FTPreset"), FTPresets)
+			.AppendChildCollection(TEXT("Revisions"), TEXT("Revision"), Revisions)
+			.AppendOptionalChild(TEXT("Protocols"), Protocols);
+
+		return ChildBuilder.GetIntermediateXmlNode();
+	}
 }
diff --git a/Engine/Plugins/VirtualProduction/DMX/DMXGDTF/Source/DMXGDTF/Private/GDTF/DMXGDTFVersion.cpp b/Engine/Plugins/VirtualProduction/DMX/DMXGDTF/Source/DMXGDTF/Private/GDTF/DMXGDTFVersion.cpp
--- a/Engine/Plugins/VirtualProduction/DMX/DMXGDTF/Source/DMXGDTF/Private/GDTF/DMXGDTFVersion.cpp
+++ b/Engine/Plugins/VirtualProduction/DMX/DMXGDTF/Source/DMXGDTF/Private/GDTF/DMXGDTFVersion.cpp
@@ -2,17 +2,20 @@
 
 #include "GDTF/DMXGDTFVersion.h"
 
-const FString FDMXGDTFVersion::GetMajorVersionAsString()
+namespace UE::DMX::GDTF
 {
-	return FString::FromInt(FDMXGDTFVersion::MajorVersion);
-}
+	const FString FDMXGDTFVersion::GetMajorVersionAsString()
+	{
+		return FString::FromInt(FDMXGDTFVersion::MajorVersion);
+	}
 
-const FString FDMXGDTFVersion::GetMinorVersionAsString()
-{
-	return FString::FromInt(FDMXGDTFVersion::MinorVersion);
-}
+	const FString FDMXGDTFVersion::GetMinorVersionAsString()
+	{
+		return FString::FromInt(FDMXGDTFVersion::MinorVersion);
+	}
 
-const FString FDMXGDTFVersion::GetAsString()
-{
-	return GetMajorVersionAsString() + TEXT(".") + GetMinorVersionAsString();
+	const FString FDMXGDTFVersion::GetAsString()
+	{
+		return GetMajorVersionAsString() + TEXT(".") + GetMinorVersionAsString();
+	}
 }
diff --git a/Engine/Plugins/VirtualProduction/DMX/DMXGDTF/Source/DMXGDTF/Private/GDTF/DMXModes/DMXGDTFChannelFunction.cpp b/Engine/Plugins/VirtualProduction/DMX/DMXGDTF/Source/DMXGDTF/Private/GDTF/DMXModes/DMXGDTFChannelFunction.cpp
--- a/Engine/Plugins/VirtualProduction/DMX/DMXGDTF/Source/DMXGDTF/Private/GDTF/DMXModes/DMXGDTFChannelFunction.cpp
+++ b/Engine/Plugins/VirtualProduction/DMX/DMXGDTF/Source/DMXGDTF/Private/GDTF/DMXModes/DMXGDTFChannelFunction.cpp
@@ -18,6 +18,7 @@
 #include "GDTF/PhysicalDescriptions/DMXGDTFPhysicalDescriptions.h"
 #include "GDTF/Wheels/DMXGDTFWheel.h"
 #include "Serialization/DMXGDTFNodeInitializer.h"
+#include "Serialization/DMXGDTFXmlNodeBuilder.h"
 
 namespace UE::DMX::GDTF
 {
@@ -32,7 +33,7 @@ namespace UE::DMX::GDTF
 			.GetAttribute(TEXT("Attribute"), Attribute)
 			.GetAttribute(TEXT("OriginalAttribute"), OriginalAttribute)
 			.GetAttribute(TEXT("DMXFrom"), DMXFrom)
-			.GetAttribute(TEXT("Default"), Default)
+			.GetAttribute(TEXT("Default"), Default, this, &FDMXGDTFChannelFunction::ParseDefault, &XmlNode)
 			.GetAttribute(TEXT("PhysicalFrom"), PhysicalFrom)
 			.GetAttribute(TEXT("PhysicalTo"), PhysicalTo)
 			.GetAttribute(TEXT("RealFade"), RealFade)
@@ -48,8 +49,57 @@ namespace UE::DMX::GDTF
 			.GetAttribute(TEXT("DMXProfile"), DMXProfile)
 			.GetAttribute(TEXT("Min"), Min)
 			.GetAttribute(TEXT("Max"), Max)
+			.GetAttribute(TEXT("CustomName"), CustomName)
 			.CreateChildren(TEXT("ChannelSet"), ChannelSetArray)
 			.CreateChildren(TEXT("SubchannelSet"), SubchannelSetArray);
+
+		// As per specs, Min and Max default the same as PhysicalFrom and PhysicalTo
+		if (Min != PhysicalFrom)
+		{
+			Min = PhysicalFrom;
+		}
+
+		if (Max != PhysicalTo)
+		{
+			Max = PhysicalTo;
+		}
+	}
+
+	FXmlNode* FDMXGDTFChannelFunction::CreateXmlNode(FXmlNode& Parent)
+	{
+		const FString DefaultLink = TEXT("");
+		const FDMXGDTFDMXValue DefaultModeFrom = TEXT("0/1");
+		const FDMXGDTFDMXValue DefaultModeTo = TEXT("0/1");
+		const float DefaultMin = PhysicalFrom;
+		const float DefaultMax = PhysicalTo;
+		const FString DefaultCustomName = TEXT("");
+
+		const FDMXGDTFXmlNodeBuilder ChildBuilder = FDMXGDTFXmlNodeBuilder(Parent, *this)
+			.SetAttribute(TEXT("Name"), Name)
+			.SetAttribute(TEXT("Attribute"), Attribute)
+			.SetAttribute(TEXT("OriginalAttribute"), OriginalAttribute)
+			.SetAttribute(TEXT("DMXFrom"), DMXFrom)
+			.SetAttribute(TEXT("Default"), Default)
+			.SetAttribute(TEXT("PhysicalFrom"), PhysicalFrom)
+			.SetAttribute(TEXT("PhysicalTo"), PhysicalTo)
+			.SetAttribute(TEXT("RealFade"), RealFade)
+			.SetAttribute(TEXT("RealAcceleration"), RealAcceleration)
+			.SetAttribute(TEXT("Wheel"), Wheel, DefaultLink)
+			.SetAttribute(TEXT("Emitter"), Emitter, DefaultLink)
+			.SetAttribute(TEXT("Filter"), Filter, DefaultLink)
+			.SetAttribute(TEXT("ColorSpace"), ColorSpace, DefaultLink)
+			.SetAttribute(TEXT("Gamut"), Gamut, DefaultLink)
+			.SetAttribute(TEXT("ModeMaster"), ModeMaster, DefaultLink)
+			.SetAttribute(TEXT("ModeFrom"), ModeFrom, DefaultModeFrom)
+			.SetAttribute(TEXT("ModeTo"), ModeTo, DefaultModeTo)
+			.SetAttribute(TEXT("DMXProfile"), DMXProfile, DefaultLink)
+			.SetAttribute(TEXT("Min"), Min, DefaultMin)
+			.SetAttribute(TEXT("Max"), Max, DefaultMax)
+			.SetAttribute(TEXT("CustomName"), CustomName, DefaultCustomName)
+			.AppendChildren(TEXT("ChannelSet"), ChannelSetArray)
+			.AppendChildren(TEXT("SubchannelSet"), SubchannelSetArray);
+
+		return ChildBuilder.GetIntermediateXmlNode();
 	}
 
 	TSharedPtr<FDMXGDTFAttribute> FDMXGDTFChannelFunction::ResolveAttribute() const
@@ -108,10 +158,7 @@ namespace UE::DMX::GDTF
 		const TSharedPtr<FDMXGDTFPhysicalDescriptions> PhysicalDescriptions = FixtureType.IsValid() ? FixtureType->PhysicalDescriptions : nullptr;
 		if (PhysicalDescriptions.IsValid())
 		{
-			if (const TSharedPtr<FDMXGDTFColorSpace>* ColorSpacePtr = Algo::FindBy(PhysicalDescriptions->ColorSpaces, ColorSpace, &FDMXGDTFColorSpace::Name))
-			{
-				return *ColorSpacePtr;
-			}
+			return PhysicalDescriptions->ColorSpaces;
 		}
 		return nullptr;
 	}
@@ -159,4 +206,26 @@ namespace UE::DMX::GDTF
 	{
 		ResolveModeMaster(OutDMXChannel, OutChannelFunction);
 	}
+	
+	FDMXGDTFDMXValue FDMXGDTFChannelFunction::ParseDefault(const FString& Value, const FXmlNode* XmlNode) const
+	{
+		if (!Value.IsEmpty())
+		{
+			return FDMXGDTFDMXValue(*Value);
+		}
+		else if(Algo::FindBy(XmlNode->GetAttributes(), TEXT("Default"), &FXmlAttribute::GetValue) == nullptr)
+		{						
+			// Try to read the value from the DXM channel instead, as it was defined in GDTF 1.0
+
+			const TSharedPtr<FDMXGDTFDMXChannel> DMXChannel = OuterLogicalChannel.IsValid() ? OuterLogicalChannel.Pin()->OuterDMXChannel.Pin() : nullptr;
+			if (ensureMsgf(DMXChannel.IsValid(), TEXT("Invalid node, node does not reside in a valid DMX channel.")))
+			{
+				PRAGMA_DISABLE_DEPRECATION_WARNINGS
+				return DMXChannel->Default;
+				PRAGMA_ENABLE_DEPRECATION_WARNINGS
+			}
+		}
+
+		return FDMXGDTFDMXValue();
+	}
 }
diff --git a/Engine/Plugins/VirtualProduction/DMX/DMXGDTF/Source/DMXGDTF/Private/GDTF/DMXModes/DMXGDTFChannelRelation.cpp b/Engine/Plugins/VirtualProduction/DMX/DMXGDTF/Source/DMXGDTF/Private/GDTF/DMXModes/DMXGDTFChannelRelation.cpp
--- a/Engine/Plugins/VirtualProduction/DMX/DMXGDTF/Source/DMXGDTF/Private/GDTF/DMXModes/DMXGDTFChannelRelation.cpp
+++ b/Engine/Plugins/VirtualProduction/DMX/DMXGDTF/Source/DMXGDTF/Private/GDTF/DMXModes/DMXGDTFChannelRelation.cpp
@@ -5,6 +5,7 @@
 #include "Algo/Find.h"
 #include "GDTF/DMXModes/DMXGDTFDMXMode.h"
 #include "Serialization/DMXGDTFNodeInitializer.h"
+#include "Serialization/DMXGDTFXmlNodeBuilder.h"
 
 namespace UE::DMX::GDTF
 {
@@ -21,6 +22,17 @@ namespace UE::DMX::GDTF
 			.GetAttribute(TEXT("Type"), Type);
 	}
 
+	FXmlNode* FDMXGDTFChannelRelation::CreateXmlNode(FXmlNode& Parent)
+	{
+		const FDMXGDTFXmlNodeBuilder ChildBuilder = FDMXGDTFXmlNodeBuilder(Parent, *this)
+			.SetAttribute(TEXT("Name"), Name)
+			.SetAttribute(TEXT("Master"), Master)
+			.SetAttribute(TEXT("Follower"), Follower)
+			.SetAttribute(TEXT("Type"), Type);
+
+		return ChildBuilder.GetIntermediateXmlNode();
+	}
+
 	TSharedPtr<FDMXGDTFDMXChannel> FDMXGDTFChannelRelation::ResolveMaster() const
 	{
 		if (const TSharedPtr<FDMXGDTFDMXMode> DMXMode = OuterDMXMode.Pin())
diff --git a/Engine/Plugins/VirtualProduction/DMX/DMXGDTF/Source/DMXGDTF/Private/GDTF/DMXModes/DMXGDTFChannelSet.cpp b/Engine/Plugins/VirtualProduction/DMX/DMXGDTF/Source/DMXGDTF/Private/GDTF/DMXModes/DMXGDTFChannelSet.cpp
--- a/Engine/Plugins/VirtualProduction/DMX/DMXGDTF/Source/DMXGDTF/Private/GDTF/DMXModes/DMXGDTFChannelSet.cpp
+++ b/Engine/Plugins/VirtualProduction/DMX/DMXGDTF/Source/DMXGDTF/Private/GDTF/DMXModes/DMXGDTFChannelSet.cpp
@@ -3,6 +3,7 @@
 #include "GDTF/DMXModes/DMXGDTFChannelSet.h"
 
 #include "Serialization/DMXGDTFNodeInitializer.h"
+#include "Serialization/DMXGDTFXmlNodeBuilder.h"
 
 namespace UE::DMX::GDTF
 {
@@ -19,4 +20,19 @@ namespace UE::DMX::GDTF
 			.GetAttribute(TEXT("PhysicalTo"), PhysicalTo)
 			.GetAttribute(TEXT("WheelSlotIndex"), WheelSlotIndex);
 	}
+
+	FXmlNode* FDMXGDTFChannelSet::CreateXmlNode(FXmlNode& Parent)
+	{
+		const float DefaultPhysicalFrom = 0.f;
+		const float DefaultPhysicalTo = 1.f;
+
+		const FDMXGDTFXmlNodeBuilder ChildBuilder = FDMXGDTFXmlNodeBuilder(Parent, *this)
+			.SetAttribute(TEXT("Name"), Name)
+			.SetAttribute(TEXT("DMXFrom"), DMXFrom)
+			.SetAttribute(TEXT("PhysicalFrom"), PhysicalFrom, DefaultPhysicalFrom)
+			.SetAttribute(TEXT("PhysicalTo"), PhysicalTo, DefaultPhysicalTo)
+			.SetAttribute(TEXT("WheelSlotIndex"), WheelSlotIndex);
+
+		return ChildBuilder.GetIntermediateXmlNode();
+	}
 }
diff --git a/Engine/Plugins/VirtualProduction/DMX/DMXGDTF/Source/DMXGDTF/Private/GDTF/DMXModes/DMXGDTFDMXChannel.cpp b/Engine/Plugins/VirtualProduction/DMX/DMXGDTF/Source/DMXGDTF/Private/GDTF/DMXModes/DMXGDTFDMXChannel.cpp
--- a/Engine/Plugins/VirtualProduction/DMX/DMXGDTF/Source/DMXGDTF/Private/GDTF/DMXModes/DMXGDTFDMXChannel.cpp
+++ b/Engine/Plugins/VirtualProduction/DMX/DMXGDTF/Source/DMXGDTF/Private/GDTF/DMXModes/DMXGDTFDMXChannel.cpp
@@ -11,6 +11,7 @@
 #include "GDTF/Geometries/DMXGDTFGeometry.h"
 #include "GDTF/Geometries/DMXGDTFGeometryCollect.h"
 #include "Serialization/DMXGDTFNodeInitializer.h"
+#include "Serialization/DMXGDTFXmlNodeBuilder.h"
 
 namespace UE::DMX::GDTF
 {
@@ -26,11 +27,29 @@ namespace UE::DMX::GDTF
 			.GetAttribute(TEXT("DMXBreak"), DMXBreak)
 			.GetAttribute(TEXT("Offset"), Offset, this, &FDMXGDTFDMXChannel::ParseOffset)
 			.GetAttribute(TEXT("InitialFunction"), InitialFunction)
-			.GetAttribute(TEXT("DMXValue"), DMXValue)
+			.GetAttribute(TEXT("Highlight"), Highlight)
 			.GetAttribute(TEXT("Geometry"), Geometry)
+PRAGMA_DISABLE_DEPRECATION_WARNINGS
+			.GetAttribute(TEXT("Default"), Default) // Deprecated with GDTF 1.1, but still initialize so old GDTFs can be supported. See DMXGDTFChannelFunction for the upgrade path.
+PRAGMA_ENABLE_DEPRECATION_WARNINGS
 			.CreateChildren(TEXT("LogicalChannel"), LogicalChannelArray);
 	}
 
+	FXmlNode* FDMXGDTFDMXChannel::CreateXmlNode(FXmlNode& Parent)
+	{
+		const FString DefaultInitialFunction = TEXT("");
+
+		const FDMXGDTFXmlNodeBuilder ChildBuilder = FDMXGDTFXmlNodeBuilder(Parent, *this)
+			.SetAttribute(TEXT("DMXBreak"), DMXBreak)
+			.SetAttribute(TEXT("Offset"), Offset)
+			.SetAttribute(TEXT("InitialFunction"), InitialFunction, DefaultInitialFunction)
+			.SetAttribute(TEXT("Highlight"), Highlight)
+			.SetAttribute(TEXT("Geometry"), Geometry)
+			.AppendChildren(TEXT("LogicalChannel"), LogicalChannelArray);
+
+		return ChildBuilder.GetIntermediateXmlNode();
+	}
+
 	TSharedPtr<FDMXGDTFChannelFunction> FDMXGDTFDMXChannel::ResolveInitialFunction() const
 	{
 		if (InitialFunction.IsEmpty() && !LogicalChannelArray.IsEmpty() && LogicalChannelArray[0].IsValid())
@@ -51,7 +70,7 @@ namespace UE::DMX::GDTF
 			}
 			else if (!LogicalChannelArray.IsEmpty() && LogicalChannelArray[0].IsValid())
 			{
-				// Fall back to the default value if the link cannot be resolved
+				// As per specs, the first channel function if no initial function is specified
 				return LogicalChannelArray[0]->ChannelFunctionArray.IsEmpty() ? nullptr : LogicalChannelArray[0]->ChannelFunctionArray[0];
 			}
 		}
diff --git a/Engine/Plugins/VirtualProduction/DMX/DMXGDTF/Source/DMXGDTF/Private/GDTF/DMXModes/DMXGDTFDMXMode.cpp b/Engine/Plugins/VirtualProduction/DMX/DMXGDTF/Source/DMXGDTF/Private/GDTF/DMXModes/DMXGDTFDMXMode.cpp
--- a/Engine/Plugins/VirtualProduction/DMX/DMXGDTF/Source/DMXGDTF/Private/GDTF/DMXModes/DMXGDTFDMXMode.cpp
+++ b/Engine/Plugins/VirtualProduction/DMX/DMXGDTF/Source/DMXGDTF/Private/GDTF/DMXModes/DMXGDTFDMXMode.cpp
@@ -10,6 +10,7 @@
 #include "GDTF/DMXModes/DMXGDTFFTMacro.h"
 #include "GDTF/DMXModes/DMXGDTFLogicalChannel.h"
 #include "Serialization/DMXGDTFNodeInitializer.h"
+#include "Serialization/DMXGDTFXmlNodeBuilder.h"
 
 namespace UE::DMX::GDTF
 {
@@ -28,6 +29,19 @@ namespace UE::DMX::GDTF
 			.CreateChildCollection(TEXT("FTMacros"), TEXT("FTMacro"), FTMacros);
 	}
 
+	FXmlNode* FDMXGDTFDMXMode::CreateXmlNode(FXmlNode& Parent)
+	{
+		const FDMXGDTFXmlNodeBuilder ChildBuilder = FDMXGDTFXmlNodeBuilder(Parent, *this)
+			.SetAttribute(TEXT("Name"), Name)
+			.SetAttribute(TEXT("Description"), Description)
+			.SetAttribute(TEXT("Geometry"), Geometry)
+			.AppendChildCollection(TEXT("DMXChannels"), TEXT("DMXChannel"), DMXChannels)
+			.AppendChildCollection(TEXT("Relations"), TEXT("Relation"), Relations)
+			.AppendChildCollection(TEXT("FTMacros"), TEXT("FTMacro"), FTMacros);
+
+		return ChildBuilder.GetIntermediateXmlNode();
+	}
+
 	void FDMXGDTFDMXMode::ResolveChannel(const FString& Link, TSharedPtr<FDMXGDTFDMXChannel>& OutDMXChannel, TSharedPtr<FDMXGDTFChannelFunction>& OutChannelFunction) const
 	{
 		if (Link.IsEmpty())
