DMX - Add GDTF export functinality, unit testing, improve GDTF data structure

diff --git a/Engine/Plugins/VirtualProduction/DMX/DMXGDTF/Source/DMXGDTF/Private/GDTF/DMXModes/DMXGDTFDMXValue.cpp b/Engine/Plugins/VirtualProduction/DMX/DMXGDTF/Source/DMXGDTF/Private/GDTF/DMXModes/DMXGDTFDMXValue.cpp
--- a/Engine/Plugins/VirtualProduction/DMX/DMXGDTF/Source/DMXGDTF/Private/GDTF/DMXModes/DMXGDTFDMXValue.cpp
+++ b/Engine/Plugins/VirtualProduction/DMX/DMXGDTF/Source/DMXGDTF/Private/GDTF/DMXModes/DMXGDTFDMXValue.cpp
@@ -4,6 +4,10 @@
 
 namespace UE::DMX::GDTF
 {
+	FDMXGDTFDMXValue::FDMXGDTFDMXValue(const TCHAR* InValue)
+		: Value(InValue)
+	{}
+
 	bool FDMXGDTFDMXValue::ToInt(uint32& OutInteger) const
 	{
 		if (Value.IsEmpty() || Value == TEXT("None"))
diff --git a/Engine/Plugins/VirtualProduction/DMX/DMXGDTF/Source/DMXGDTF/Private/GDTF/DMXModes/DMXGDTFFTMacro.cpp b/Engine/Plugins/VirtualProduction/DMX/DMXGDTF/Source/DMXGDTF/Private/GDTF/DMXModes/DMXGDTFFTMacro.cpp
--- a/Engine/Plugins/VirtualProduction/DMX/DMXGDTF/Source/DMXGDTF/Private/GDTF/DMXModes/DMXGDTFFTMacro.cpp
+++ b/Engine/Plugins/VirtualProduction/DMX/DMXGDTF/Source/DMXGDTF/Private/GDTF/DMXModes/DMXGDTFFTMacro.cpp
@@ -5,6 +5,7 @@
 #include "GDTF/DMXModes/DMXGDTFDMXMode.h"
 #include "GDTF/DMXModes/DMXGDTFMacroDMX.h"
 #include "Serialization/DMXGDTFNodeInitializer.h"
+#include "Serialization/DMXGDTFXmlNodeBuilder.h"
 
 namespace UE::DMX::GDTF
 {
@@ -20,6 +21,18 @@ namespace UE::DMX::GDTF
 			.CreateChildren(TEXT("MacroDMX"), MacroDMXArray);
 	}
 
+	FXmlNode* FDMXGDTFFTMacro::CreateXmlNode(FXmlNode& Parent)
+	{
+		const FString DefaultChannelFunction = TEXT("");
+
+		const FDMXGDTFXmlNodeBuilder ChildBuilder = FDMXGDTFXmlNodeBuilder(Parent, *this)
+			.SetAttribute(TEXT("Name"), Name)
+			.SetAttribute(TEXT("ChannelFunction"), ChannelFunction, DefaultChannelFunction)
+			.AppendChildren(TEXT("MacroDMX"), MacroDMXArray);
+
+		return ChildBuilder.GetIntermediateXmlNode();
+	}
+
 	TSharedPtr<FDMXGDTFChannelFunction> FDMXGDTFFTMacro::ResolveChannelFunction() const
 	{
 		if (const TSharedPtr<FDMXGDTFDMXMode> DMXMode = OuterDMXMode.Pin())
diff --git a/Engine/Plugins/VirtualProduction/DMX/DMXGDTF/Source/DMXGDTF/Private/GDTF/DMXModes/DMXGDTFLogicalChannel.cpp b/Engine/Plugins/VirtualProduction/DMX/DMXGDTF/Source/DMXGDTF/Private/GDTF/DMXModes/DMXGDTFLogicalChannel.cpp
--- a/Engine/Plugins/VirtualProduction/DMX/DMXGDTF/Source/DMXGDTF/Private/GDTF/DMXModes/DMXGDTFLogicalChannel.cpp
+++ b/Engine/Plugins/VirtualProduction/DMX/DMXGDTF/Source/DMXGDTF/Private/GDTF/DMXModes/DMXGDTFLogicalChannel.cpp
@@ -8,6 +8,7 @@
 #include "GDTF/DMXGDTFFixtureType.h"
 #include "GDTF/DMXModes/DMXGDTFChannelFunction.h"
 #include "Serialization/DMXGDTFNodeInitializer.h"
+#include "Serialization/DMXGDTFXmlNodeBuilder.h"
 
 namespace UE::DMX::GDTF
 {
@@ -26,6 +27,19 @@ namespace UE::DMX::GDTF
 			.CreateChildren(TEXT("ChannelFunction"), ChannelFunctionArray);
 	}
 
+	FXmlNode* FDMXGDTFLogicalChannel::CreateXmlNode(FXmlNode& Parent)
+	{
+		const FDMXGDTFXmlNodeBuilder ChildBuilder = FDMXGDTFXmlNodeBuilder(Parent, *this)
+			.SetAttribute(TEXT("Attribute"), Attribute)
+			.SetAttribute(TEXT("Snap"), Snap)
+			.SetAttribute(TEXT("Master"), Master)
+			.SetAttribute(TEXT("MibFade"), MibFade)
+			.SetAttribute(TEXT("DMXChangeTimeLimit"), DMXChangeTimeLimit)
+			.AppendChildren(TEXT("ChannelFunction"), ChannelFunctionArray);
+
+		return ChildBuilder.GetIntermediateXmlNode();
+	}
+
 	TSharedPtr<FDMXGDTFAttribute> FDMXGDTFLogicalChannel::ResolveAttribute() const
 	{
 		const TSharedPtr<FDMXGDTFFixtureType> FixtureType = GetFixtureType().Pin();
diff --git a/Engine/Plugins/VirtualProduction/DMX/DMXGDTF/Source/DMXGDTF/Private/GDTF/DMXModes/DMXGDTFMacroDMX.cpp b/Engine/Plugins/VirtualProduction/DMX/DMXGDTF/Source/DMXGDTF/Private/GDTF/DMXModes/DMXGDTFMacroDMX.cpp
--- a/Engine/Plugins/VirtualProduction/DMX/DMXGDTF/Source/DMXGDTF/Private/GDTF/DMXModes/DMXGDTFMacroDMX.cpp
+++ b/Engine/Plugins/VirtualProduction/DMX/DMXGDTF/Source/DMXGDTF/Private/GDTF/DMXModes/DMXGDTFMacroDMX.cpp
@@ -4,6 +4,7 @@
 
 #include "GDTF/DMXModes/DMXGDTFMacroDMXStep.h"
 #include "Serialization/DMXGDTFNodeInitializer.h"
+#include "Serialization/DMXGDTFXmlNodeBuilder.h"
 
 namespace UE::DMX::GDTF
 {
@@ -16,4 +17,12 @@ namespace UE::DMX::GDTF
 		FDMXGDTFNodeInitializer(SharedThis(this), XmlNode)
 			.CreateChildren(TEXT("MacroDMXStep"), MacroDMXStepArray);
 	}
+
+	FXmlNode* FDMXGDTFMacroDMX::CreateXmlNode(FXmlNode& Parent)
+	{
+		const FDMXGDTFXmlNodeBuilder ChildBuilder = FDMXGDTFXmlNodeBuilder(Parent, *this)
+			.AppendChildren(TEXT("MacroDMXStep"), MacroDMXStepArray);
+
+		return ChildBuilder.GetIntermediateXmlNode();
+	}
 }
diff --git a/Engine/Plugins/VirtualProduction/DMX/DMXGDTF/Source/DMXGDTF/Private/GDTF/DMXModes/DMXGDTFMacroDMXStep.cpp b/Engine/Plugins/VirtualProduction/DMX/DMXGDTF/Source/DMXGDTF/Private/GDTF/DMXModes/DMXGDTFMacroDMXStep.cpp
--- a/Engine/Plugins/VirtualProduction/DMX/DMXGDTF/Source/DMXGDTF/Private/GDTF/DMXModes/DMXGDTFMacroDMXStep.cpp
+++ b/Engine/Plugins/VirtualProduction/DMX/DMXGDTF/Source/DMXGDTF/Private/GDTF/DMXModes/DMXGDTFMacroDMXStep.cpp
@@ -4,6 +4,7 @@
 
 #include "GDTF/DMXModes/DMXGDTFMacroDMXValue.h"
 #include "Serialization/DMXGDTFNodeInitializer.h"
+#include "Serialization/DMXGDTFXmlNodeBuilder.h"
 
 namespace UE::DMX::GDTF
 {
@@ -17,4 +18,13 @@ namespace UE::DMX::GDTF
 			.GetAttribute(TEXT("Duration"), Duration)
 			.CreateChildren(TEXT("MacroDMXValue"), MacroDMXValueArray);
 	}
+
+	FXmlNode* FDMXGDTFMacroDMXStep::CreateXmlNode(FXmlNode& Parent)
+	{
+		const FDMXGDTFXmlNodeBuilder ChildBuilder = FDMXGDTFXmlNodeBuilder(Parent, *this)
+			.SetAttribute(TEXT("Duration"), Duration)
+			.AppendChildren(TEXT("MacroDMXValue"), MacroDMXValueArray);
+
+		return ChildBuilder.GetIntermediateXmlNode();
+	}
 }
diff --git a/Engine/Plugins/VirtualProduction/DMX/DMXGDTF/Source/DMXGDTF/Private/GDTF/DMXModes/DMXGDTFMacroDMXValue.cpp b/Engine/Plugins/VirtualProduction/DMX/DMXGDTF/Source/DMXGDTF/Private/GDTF/DMXModes/DMXGDTFMacroDMXValue.cpp
--- a/Engine/Plugins/VirtualProduction/DMX/DMXGDTF/Source/DMXGDTF/Private/GDTF/DMXModes/DMXGDTFMacroDMXValue.cpp
+++ b/Engine/Plugins/VirtualProduction/DMX/DMXGDTF/Source/DMXGDTF/Private/GDTF/DMXModes/DMXGDTFMacroDMXValue.cpp
@@ -8,6 +8,7 @@
 #include "GDTF/DMXModes/DMXGDTFMacroDMX.h"
 #include "GDTF/DMXModes/DMXGDTFMacroDMXStep.h"
 #include "Serialization/DMXGDTFNodeInitializer.h"
+#include "Serialization/DMXGDTFXmlNodeBuilder.h"
 
 namespace UE::DMX::GDTF
 {
@@ -22,6 +23,14 @@ namespace UE::DMX::GDTF
 			.GetAttribute(TEXT("DMXChannel"), DMXChannel);
 	}
 
+	FXmlNode* FDMXGDTFMacroDMXValue::CreateXmlNode(FXmlNode& Parent)
+	{
+		const FDMXGDTFXmlNodeBuilder ChildBuilder = FDMXGDTFXmlNodeBuilder(Parent, *this)
+			.SetAttribute(TEXT("DMXValue"), DMXValue)
+			.SetAttribute(TEXT("DMXChannel"), DMXChannel);
+
+		return ChildBuilder.GetIntermediateXmlNode();
+	}
 	TSharedPtr<FDMXGDTFDMXChannel> FDMXGDTFMacroDMXValue::ResolveDMXChannel() const
 	{
 		const TSharedPtr<FDMXGDTFMacroDMX> MacroDMX = OuterMacroDMXStep.IsValid() ? OuterMacroDMXStep.Pin()->OuterMacroDMX.Pin() : nullptr;
diff --git a/Engine/Plugins/VirtualProduction/DMX/DMXGDTF/Source/DMXGDTF/Private/GDTF/DMXModes/DMXGDTFSubchannelSet.cpp b/Engine/Plugins/VirtualProduction/DMX/DMXGDTF/Source/DMXGDTF/Private/GDTF/DMXModes/DMXGDTFSubchannelSet.cpp
--- a/Engine/Plugins/VirtualProduction/DMX/DMXGDTF/Source/DMXGDTF/Private/GDTF/DMXModes/DMXGDTFSubchannelSet.cpp
+++ b/Engine/Plugins/VirtualProduction/DMX/DMXGDTF/Source/DMXGDTF/Private/GDTF/DMXModes/DMXGDTFSubchannelSet.cpp
@@ -10,6 +10,7 @@
 #include "GDTF/PhysicalDescriptions/DMXGDTFDMXProfile.h"
 #include "GDTF/PhysicalDescriptions/DMXGDTFPhysicalDescriptions.h"
 #include "Serialization/DMXGDTFNodeInitializer.h"
+#include "Serialization/DMXGDTFXmlNodeBuilder.h"
 
 namespace UE::DMX::GDTF
 {
@@ -27,6 +28,20 @@ namespace UE::DMX::GDTF
 			.GetAttribute(TEXT("DMXProfile"), DMXProfile);
 	}
 
+	FXmlNode* FDMXGDTFSubchannelSet::CreateXmlNode(FXmlNode& Parent)
+	{
+		const float DefaultPhysicalFrom = 0.f;
+		const float DefaultPhysicalTo = 1.f;
+
+		const FDMXGDTFXmlNodeBuilder ChildBuilder = FDMXGDTFXmlNodeBuilder(Parent, *this)
+			.SetAttribute(TEXT("Name"), Name)
+			.SetAttribute(TEXT("PhysicalFrom"), PhysicalFrom, DefaultPhysicalFrom)
+			.SetAttribute(TEXT("PhysicalTo"), PhysicalTo, DefaultPhysicalTo)
+			.SetAttribute(TEXT("SubphyiscalUnit"), SubphyiscalUnit)
+			.SetAttribute(TEXT("DMXProfile"), DMXProfile);
+
+		return ChildBuilder.GetIntermediateXmlNode();
+	}
 	TSharedPtr<FDMXGDTFSubphysicalUnit> FDMXGDTFSubchannelSet::ResolveSubphysicalUnit() const
 	{
 		TArray<FString> Link;
diff --git a/Engine/Plugins/VirtualProduction/DMX/DMXGDTF/Source/DMXGDTF/Private/GDTF/Geometries/DMXGDTFBeamGeometry.cpp b/Engine/Plugins/VirtualProduction/DMX/DMXGDTF/Source/DMXGDTF/Private/GDTF/Geometries/DMXGDTFBeamGeometry.cpp
--- a/Engine/Plugins/VirtualProduction/DMX/DMXGDTF/Source/DMXGDTF/Private/GDTF/Geometries/DMXGDTFBeamGeometry.cpp
+++ b/Engine/Plugins/VirtualProduction/DMX/DMXGDTF/Source/DMXGDTF/Private/GDTF/Geometries/DMXGDTFBeamGeometry.cpp
@@ -6,6 +6,7 @@
 #include "GDTF/PhysicalDescriptions/DMXGDTFEmitter.h"
 #include "GDTF/PhysicalDescriptions/DMXGDTFPhysicalDescriptions.h"
 #include "Serialization/DMXGDTFNodeInitializer.h"
+#include "Serialization/DMXGDTFXmlNodeBuilder.h"
 
 namespace UE::DMX::GDTF
 {
@@ -19,15 +20,40 @@ namespace UE::DMX::GDTF
 			.GetAttribute(TEXT("LuminousFlux"), LuminousFlux)
 			.GetAttribute(TEXT("ColorTemperature"), ColorTemperature)
 			.GetAttribute(TEXT("BeamAngle"), BeamAngle)
-			.GetAttribute(TEXT("FieldAnge"), FieldAngle)
+			.GetAttribute(TEXT("FieldAngle"), FieldAngle)
 			.GetAttribute(TEXT("ThrowRatio"), ThrowRatio)
 			.GetAttribute(TEXT("RectangleRatio"), RectangleRatio)
 			.GetAttribute(TEXT("BeamRadius"), BeamRadius)
-			.GetAttribute(TEXT("LampType"), BeamType)
+			.GetAttribute(TEXT("BeamType"), BeamType)
 			.GetAttribute(TEXT("ColorRenderingIndex"), ColorRenderingIndex)
 			.GetAttribute(TEXT("EmitterSpectrum"), EmitterSpectrum);
 	}
 
+	FXmlNode* FDMXGDTFBeamGeometry::CreateXmlNode(FXmlNode& Parent)
+	{
+		const float DefaultThrowRatio = 1.f;
+		const float DefaultRecatangleRatio = 1.7777f;
+		const FString DefaultEmitterSpectrum = TEXT("");
+
+		FXmlNode* AppendToNode = FDMXGDTFGeometry::CreateXmlNode(Parent);
+
+		const FDMXGDTFXmlNodeBuilder ChildBuilder = FDMXGDTFXmlNodeBuilder(Parent, *this, AppendToNode)
+			.SetAttribute(TEXT("LampType"), LampType)
+			.SetAttribute(TEXT("PowerConsumption"), PowerConsumption)
+			.SetAttribute(TEXT("LuminousFlux"), LuminousFlux)
+			.SetAttribute(TEXT("ColorTemperature"), ColorTemperature)
+			.SetAttribute(TEXT("BeamAngle"), BeamAngle)
+			.SetAttribute(TEXT("FieldAngle"), FieldAngle)
+			.SetAttribute(TEXT("ThrowRatio"), ThrowRatio, DefaultThrowRatio)
+			.SetAttribute(TEXT("RectangleRatio"), RectangleRatio, DefaultRecatangleRatio)
+			.SetAttribute(TEXT("BeamRadius"), BeamRadius)
+			.SetAttribute(TEXT("BeamType"), BeamType)
+			.SetAttribute(TEXT("ColorRenderingIndex"), ColorRenderingIndex)
+			.SetAttribute(TEXT("EmitterSpectrum"), EmitterSpectrum, DefaultEmitterSpectrum);
+
+		return ChildBuilder.GetIntermediateXmlNode();
+	}
+
 	TSharedPtr<FDMXGDTFEmitter> FDMXGDTFBeamGeometry::ResolveEmitterSpectrum() const
 	{
 		const TSharedPtr<FDMXGDTFFixtureType> FixtureType = GetFixtureType().Pin();
diff --git a/Engine/Plugins/VirtualProduction/DMX/DMXGDTF/Source/DMXGDTF/Private/GDTF/Geometries/DMXGDTFBreak.cpp b/Engine/Plugins/VirtualProduction/DMX/DMXGDTF/Source/DMXGDTF/Private/GDTF/Geometries/DMXGDTFBreak.cpp
--- a/Engine/Plugins/VirtualProduction/DMX/DMXGDTF/Source/DMXGDTF/Private/GDTF/Geometries/DMXGDTFBreak.cpp
+++ b/Engine/Plugins/VirtualProduction/DMX/DMXGDTF/Source/DMXGDTF/Private/GDTF/Geometries/DMXGDTFBreak.cpp
@@ -3,6 +3,7 @@
 #include "GDTF/Geometries/DMXGDTFBreak.h"
 
 #include "Serialization/DMXGDTFNodeInitializer.h"
+#include "Serialization/DMXGDTFXmlNodeBuilder.h"
 
 namespace UE::DMX::GDTF
 {
@@ -13,7 +14,16 @@ namespace UE::DMX::GDTF
 	void FDMXGDTFBreak::Initialize(const FXmlNode& XmlNode)
 	{
 		FDMXGDTFNodeInitializer(SharedThis(this), XmlNode)
-			.GetAttribute(TEXT("DMXAddress"), DMXAddress)
+			.GetAttribute(TEXT("DMXOffset"), DMXOffset)
 			.GetAttribute(TEXT("DMXBreak"), DMXBreak);
 	}
+
+	FXmlNode* FDMXGDTFBreak::CreateXmlNode(FXmlNode& Parent)
+	{
+		const FDMXGDTFXmlNodeBuilder ChildBuilder = FDMXGDTFXmlNodeBuilder(Parent, *this)
+			.SetAttribute(TEXT("DMXOffset"), DMXOffset)
+			.SetAttribute(TEXT("DMXBreak"), DMXBreak);
+
+		return ChildBuilder.GetIntermediateXmlNode();
+	}
 }
diff --git a/Engine/Plugins/VirtualProduction/DMX/DMXGDTF/Source/DMXGDTF/Private/GDTF/Geometries/DMXGDTFDisplayGeometry.cpp b/Engine/Plugins/VirtualProduction/DMX/DMXGDTF/Source/DMXGDTF/Private/GDTF/Geometries/DMXGDTFDisplayGeometry.cpp
--- a/Engine/Plugins/VirtualProduction/DMX/DMXGDTF/Source/DMXGDTF/Private/GDTF/Geometries/DMXGDTFDisplayGeometry.cpp
+++ b/Engine/Plugins/VirtualProduction/DMX/DMXGDTF/Source/DMXGDTF/Private/GDTF/Geometries/DMXGDTFDisplayGeometry.cpp
@@ -3,6 +3,7 @@
 #include "GDTF/Geometries/DMXGDTFDisplayGeometry.h"
 
 #include "Serialization/DMXGDTFNodeInitializer.h"
+#include "Serialization/DMXGDTFXmlNodeBuilder.h"
 
 namespace UE::DMX::GDTF
 {
@@ -13,4 +14,14 @@ namespace UE::DMX::GDTF
 		FDMXGDTFNodeInitializer(SharedThis(this), XmlNode)
 			.GetAttribute(TEXT("Texture"), Texture);
 	}
+
+	FXmlNode* FDMXGDTFDisplayGeometry::CreateXmlNode(FXmlNode& Parent)
+	{
+		FXmlNode* AppendToNode = FDMXGDTFGeometry::CreateXmlNode(Parent);
+
+		const FDMXGDTFXmlNodeBuilder ChildBuilder = FDMXGDTFXmlNodeBuilder(Parent, *this, AppendToNode)
+			.SetAttribute(TEXT("Texture"), Texture);
+
+		return ChildBuilder.GetIntermediateXmlNode();
+	}
 }
diff --git a/Engine/Plugins/VirtualProduction/DMX/DMXGDTF/Source/DMXGDTF/Private/GDTF/Geometries/DMXGDTFGeometry.cpp b/Engine/Plugins/VirtualProduction/DMX/DMXGDTF/Source/DMXGDTF/Private/GDTF/Geometries/DMXGDTFGeometry.cpp
--- a/Engine/Plugins/VirtualProduction/DMX/DMXGDTF/Source/DMXGDTF/Private/GDTF/Geometries/DMXGDTFGeometry.cpp
+++ b/Engine/Plugins/VirtualProduction/DMX/DMXGDTF/Source/DMXGDTF/Private/GDTF/Geometries/DMXGDTFGeometry.cpp
@@ -6,6 +6,7 @@
 #include "GDTF/DMXGDTFFixtureType.h"
 #include "GDTF/Models/DMXGDTFModel.h"
 #include "Serialization/DMXGDTFNodeInitializer.h"
+#include "Serialization/DMXGDTFXmlNodeBuilder.h"
 
 namespace UE::DMX::GDTF
 {
@@ -23,6 +24,20 @@ namespace UE::DMX::GDTF
 			.GetAttribute(TEXT("Position"), Position);
 	}
 
+	FXmlNode* FDMXGDTFGeometry::CreateXmlNode(FXmlNode& Parent)
+	{
+		const FString DefaultModel = TEXT("");
+
+		FXmlNode* AppendToNode = FDMXGDTFGeometryCollectBase::CreateXmlNode(Parent);
+
+		const FDMXGDTFXmlNodeBuilder ChildBuilder = FDMXGDTFXmlNodeBuilder(Parent, *this, AppendToNode)
+			.SetAttribute(TEXT("Name"), Name)
+			.SetAttribute(TEXT("Model"), Model, DefaultModel)
+			.SetAttribute(TEXT("Position"), Position, EDMXGDTFMatrixType::Matrix4x4);
+
+		return ChildBuilder.GetIntermediateXmlNode();
+	}
+
 	TSharedPtr<FDMXGDTFModel> FDMXGDTFGeometry::ResolveModel() const
 	{
 		if (const TSharedPtr<FDMXGDTFFixtureType> FixtureType = GetFixtureType().Pin())
diff --git a/Engine/Plugins/VirtualProduction/DMX/DMXGDTF/Source/DMXGDTF/Private/GDTF/Geometries/DMXGDTFGeometryCollect.cpp b/Engine/Plugins/VirtualProduction/DMX/DMXGDTF/Source/DMXGDTF/Private/GDTF/Geometries/DMXGDTFGeometryCollect.cpp
--- a/Engine/Plugins/VirtualProduction/DMX/DMXGDTF/Source/DMXGDTF/Private/GDTF/Geometries/DMXGDTFGeometryCollect.cpp
+++ b/Engine/Plugins/VirtualProduction/DMX/DMXGDTF/Source/DMXGDTF/Private/GDTF/Geometries/DMXGDTFGeometryCollect.cpp
@@ -3,6 +3,7 @@
 #include "GDTF/Geometries/DMXGDTFGeometryCollect.h"
 
 #include "Serialization/DMXGDTFNodeInitializer.h"
+#include "Serialization/DMXGDTFXmlNodeBuilder.h"
 
 namespace UE::DMX::GDTF
 {
@@ -14,4 +15,9 @@ namespace UE::DMX::GDTF
 	{
 		FDMXGDTFGeometryCollectBase::Initialize(XmlNode);
 	}
+
+	FXmlNode* FDMXGDTFGeometryCollect::CreateXmlNode(FXmlNode& Parent)
+	{
+		return FDMXGDTFGeometryCollectBase::CreateXmlNode(Parent);
+	}
 }
diff --git a/Engine/Plugins/VirtualProduction/DMX/DMXGDTF/Source/DMXGDTF/Private/GDTF/Geometries/DMXGDTFGeometryCollectBase.cpp b/Engine/Plugins/VirtualProduction/DMX/DMXGDTF/Source/DMXGDTF/Private/GDTF/Geometries/DMXGDTFGeometryCollectBase.cpp
--- a/Engine/Plugins/VirtualProduction/DMX/DMXGDTF/Source/DMXGDTF/Private/GDTF/Geometries/DMXGDTFGeometryCollectBase.cpp
+++ b/Engine/Plugins/VirtualProduction/DMX/DMXGDTF/Source/DMXGDTF/Private/GDTF/Geometries/DMXGDTFGeometryCollectBase.cpp
@@ -22,6 +22,7 @@
 #include "GDTF/Geometries/DMXGDTFSupportGeometry.h"
 #include "GDTF/Geometries/DMXGDTFWiringObjectGeometry.h"
 #include "Serialization/DMXGDTFNodeInitializer.h"
+#include "Serialization/DMXGDTFXmlNodeBuilder.h"
 
 namespace UE::DMX::GDTF
 {
@@ -54,6 +55,30 @@ namespace UE::DMX::GDTF
 			.CreateChildren(TEXT("Magnet"), MagnetArray);
 	}
 
+	FXmlNode* FDMXGDTFGeometryCollectBase::CreateXmlNode(FXmlNode& Parent)
+	{
+		const FDMXGDTFXmlNodeBuilder ChildBuilder = FDMXGDTFXmlNodeBuilder(Parent, *this)
+			.AppendChildren(TEXT("Geometry"), GeometryArray)
+			.AppendChildren(TEXT("Axis"), AxisArray)
+			.AppendChildren(TEXT("FilterBeam"), FilterBeamArray)
+			.AppendChildren(TEXT("FilterColor"), FilterColorArray)
+			.AppendChildren(TEXT("FilterGobo"), FilterGoboArray)
+			.AppendChildren(TEXT("FilterShaper"), FilterShaperArray)
+			.AppendChildren(TEXT("Beam"), BeamArray)
+			.AppendChildren(TEXT("MediaServerLayer"), MediaServerLayerArray)
+			.AppendChildren(TEXT("MediaServerCamera"), MediaServerCameraArray)
+			.AppendChildren(TEXT("MediaServerMaster"), MediaServerMasterArray)
+			.AppendChildren(TEXT("Display"), DisplayArray)
+			.AppendChildren(TEXT("GeometryReference"), GeometryReferenceArray)
+			.AppendChildren(TEXT("Laser"), LaserArray)
+			.AppendChildren(TEXT("WiringObject"), WiringObjectArray)
+			.AppendChildren(TEXT("Inventory"), InventoryArray)
+			.AppendChildren(TEXT("Structure"), StructureArray)
+			.AppendChildren(TEXT("Support"), SupportArray)
+			.AppendChildren(TEXT("Magnet"), MagnetArray);
+
+		return ChildBuilder.GetIntermediateXmlNode();
+	}
 	void FDMXGDTFGeometryCollectBase::FindGeometryByName(const TCHAR* InName, TSharedPtr<FDMXGDTFGeometry>& OutGeometry, TSharedPtr<FDMXGDTFGeometryReference>& OutGeometryReference) const
 	{
 		const TSharedPtr<FDMXGDTFGeometryReference>* GeometryReferencePtr = Algo::FindBy(GeometryReferenceArray, InName, &FDMXGDTFGeometryReference::Name);
diff --git a/Engine/Plugins/VirtualProduction/DMX/DMXGDTF/Source/DMXGDTF/Private/GDTF/Geometries/DMXGDTFGeometryReference.cpp b/Engine/Plugins/VirtualProduction/DMX/DMXGDTF/Source/DMXGDTF/Private/GDTF/Geometries/DMXGDTFGeometryReference.cpp
--- a/Engine/Plugins/VirtualProduction/DMX/DMXGDTF/Source/DMXGDTF/Private/GDTF/Geometries/DMXGDTFGeometryReference.cpp
+++ b/Engine/Plugins/VirtualProduction/DMX/DMXGDTF/Source/DMXGDTF/Private/GDTF/Geometries/DMXGDTFGeometryReference.cpp
@@ -9,6 +9,7 @@
 #include "GDTF/Geometries/DMXGDTFGeometryCollect.h"
 #include "GDTF/Models/DMXGDTFModel.h"
 #include "Serialization/DMXGDTFNodeInitializer.h"
+#include "Serialization/DMXGDTFXmlNodeBuilder.h"
 
 namespace UE::DMX::GDTF
 {
@@ -18,15 +19,28 @@ namespace UE::DMX::GDTF
 
 	void FDMXGDTFGeometryReference::Initialize(const FXmlNode& XmlNode)
 	{
-		using namespace UE::DMX::GDTF;
-
 		FDMXGDTFNodeInitializer(SharedThis(this), XmlNode)
 			.GetAttribute(TEXT("Name"), Name)
+			.GetAttribute(TEXT("Position"), Position)
 			.CreateChildren(TEXT("Break"), BreakArray)
 			.GetAttribute(TEXT("Geometry"), Geometry)
 			.GetAttribute(TEXT("Model"), Model);
 	}
 
+	FXmlNode* FDMXGDTFGeometryReference::CreateXmlNode(FXmlNode& Parent)
+	{
+		const FString DefaultModel = TEXT("");
+
+		const FDMXGDTFXmlNodeBuilder ChildBuilder = FDMXGDTFXmlNodeBuilder(Parent, *this)
+			.SetAttribute(TEXT("Name"), Name)
+			.SetAttribute(TEXT("Position"), Position, EDMXGDTFMatrixType::Matrix4x4)
+			.AppendChildren(TEXT("Break"), BreakArray)
+			.SetAttribute(TEXT("Geometry"), Geometry)
+			.SetAttribute(TEXT("Model"), Model, DefaultModel);
+
+		return ChildBuilder.GetIntermediateXmlNode();
+	}
+
 	TSharedPtr<FDMXGDTFGeometry> FDMXGDTFGeometryReference::ResolveGeometry() const
 	{
 		// Only top level geometries are allowed
diff --git a/Engine/Plugins/VirtualProduction/DMX/DMXGDTF/Source/DMXGDTF/Private/GDTF/Geometries/DMXGDTFInventoryGeometry.cpp b/Engine/Plugins/VirtualProduction/DMX/DMXGDTF/Source/DMXGDTF/Private/GDTF/Geometries/DMXGDTFInventoryGeometry.cpp
--- a/Engine/Plugins/VirtualProduction/DMX/DMXGDTF/Source/DMXGDTF/Private/GDTF/Geometries/DMXGDTFInventoryGeometry.cpp
+++ b/Engine/Plugins/VirtualProduction/DMX/DMXGDTF/Source/DMXGDTF/Private/GDTF/Geometries/DMXGDTFInventoryGeometry.cpp
@@ -3,6 +3,7 @@
 #include "GDTF/Geometries/DMXGDTFInventoryGeometry.h"
 
 #include "Serialization/DMXGDTFNodeInitializer.h"
+#include "Serialization/DMXGDTFXmlNodeBuilder.h"
 
 namespace UE::DMX::GDTF
 {
@@ -10,9 +11,17 @@ namespace UE::DMX::GDTF
 	{
 		FDMXGDTFGeometry::Initialize(XmlNode);
 
-		using namespace UE::DMX::GDTF;
-
 		FDMXGDTFNodeInitializer(SharedThis(this), XmlNode)
 			.GetAttribute(TEXT("Cout"), Count);
 	}
+
+	FXmlNode* FDMXGDTFInventoryGeometry::CreateXmlNode(FXmlNode& Parent)
+	{
+		FXmlNode* AppendToNode = FDMXGDTFGeometry::CreateXmlNode(Parent);
+
+		const FDMXGDTFXmlNodeBuilder ChildBuilder = FDMXGDTFXmlNodeBuilder(Parent, *this, AppendToNode)
+			.SetAttribute(TEXT("Cout"), Count);
+
+		return ChildBuilder.GetIntermediateXmlNode();
+	}
 }
diff --git a/Engine/Plugins/VirtualProduction/DMX/DMXGDTF/Source/DMXGDTF/Private/GDTF/Geometries/DMXGDTFLaserGeometry.cpp b/Engine/Plugins/VirtualProduction/DMX/DMXGDTF/Source/DMXGDTF/Private/GDTF/Geometries/DMXGDTFLaserGeometry.cpp
--- a/Engine/Plugins/VirtualProduction/DMX/DMXGDTF/Source/DMXGDTF/Private/GDTF/Geometries/DMXGDTFLaserGeometry.cpp
+++ b/Engine/Plugins/VirtualProduction/DMX/DMXGDTF/Source/DMXGDTF/Private/GDTF/Geometries/DMXGDTFLaserGeometry.cpp
@@ -6,6 +6,7 @@
 #include "GDTF/PhysicalDescriptions/DMXGDTFEmitter.h"
 #include "GDTF/PhysicalDescriptions/DMXGDTFPhysicalDescriptions.h"
 #include "Serialization/DMXGDTFNodeInitializer.h"
+#include "Serialization/DMXGDTFXmlNodeBuilder.h"
 
 namespace UE::DMX::GDTF
 {
@@ -29,6 +30,26 @@ namespace UE::DMX::GDTF
 			.CreateChildren(TEXT("Protocol"), ProtocolArray);
 	}
 
+	FXmlNode* FDMXGDTFLaserGeometry::CreateXmlNode(FXmlNode& Parent)
+	{
+		FXmlNode* AppendToNode = FDMXGDTFGeometry::CreateXmlNode(Parent);
+
+		const FDMXGDTFXmlNodeBuilder ChildBuilder = FDMXGDTFXmlNodeBuilder(Parent, *this, AppendToNode)
+			.SetAttribute(TEXT("ColorType"), ColorType)
+			.SetAttribute(TEXT("Color"), Color)
+			.SetAttribute(TEXT("OutputStrength"), OutputStrength)
+			.SetAttribute(TEXT("BeamDiameter"), BeamDiameter)
+			.SetAttribute(TEXT("Emitter"), Emitter)
+			.SetAttribute(TEXT("BeamDivergenceMin"), BeamDivergenceMin)
+			.SetAttribute(TEXT("BeamDivergenceMax"), BeamDivergenceMax)
+			.SetAttribute(TEXT("ScanAnglePan"), ScanAnglePan)
+			.SetAttribute(TEXT("ScanAngleTilt"), ScanAngleTilt)
+			.SetAttribute(TEXT("ScanSpeed"), ScanSpeed)
+			.AppendChildren(TEXT("Protocol"), ProtocolArray);
+
+		return ChildBuilder.GetIntermediateXmlNode();
+	}
+
 	TSharedPtr<FDMXGDTFEmitter> FDMXGDTFLaserGeometry::ResolveEmitter() const
 	{
 		const TSharedPtr<FDMXGDTFFixtureType> FixtureType = GetFixtureType().Pin();
diff --git a/Engine/Plugins/VirtualProduction/DMX/DMXGDTF/Source/DMXGDTF/Private/GDTF/Geometries/DMXGDTFLaserProtocol.cpp b/Engine/Plugins/VirtualProduction/DMX/DMXGDTF/Source/DMXGDTF/Private/GDTF/Geometries/DMXGDTFLaserProtocol.cpp
--- a/Engine/Plugins/VirtualProduction/DMX/DMXGDTF/Source/DMXGDTF/Private/GDTF/Geometries/DMXGDTFLaserProtocol.cpp
+++ b/Engine/Plugins/VirtualProduction/DMX/DMXGDTF/Source/DMXGDTF/Private/GDTF/Geometries/DMXGDTFLaserProtocol.cpp
@@ -3,6 +3,7 @@
 #include "GDTF/Geometries/DMXGDTFLaserProtocol.h"
 
 #include "Serialization/DMXGDTFNodeInitializer.h"
+#include "Serialization/DMXGDTFXmlNodeBuilder.h"
 
 namespace UE::DMX::GDTF
 {
@@ -15,4 +16,12 @@ namespace UE::DMX::GDTF
 		FDMXGDTFNodeInitializer(SharedThis(this), XmlNode)
 			.GetAttribute(TEXT("Name"), Name);
 	}
+	
+	FXmlNode* FDMXGDTFLaserProtocol::CreateXmlNode(FXmlNode& Parent)
+	{
+		const FDMXGDTFXmlNodeBuilder ChildBuilder = FDMXGDTFXmlNodeBuilder(Parent, *this)
+			.SetAttribute(TEXT("Name"), Name);
+
+		return ChildBuilder.GetIntermediateXmlNode();
+	}
 }
diff --git a/Engine/Plugins/VirtualProduction/DMX/DMXGDTF/Source/DMXGDTF/Private/GDTF/Geometries/DMXGDTFStructureGeometry.cpp b/Engine/Plugins/VirtualProduction/DMX/DMXGDTF/Source/DMXGDTF/Private/GDTF/Geometries/DMXGDTFStructureGeometry.cpp
--- a/Engine/Plugins/VirtualProduction/DMX/DMXGDTF/Source/DMXGDTF/Private/GDTF/Geometries/DMXGDTFStructureGeometry.cpp
+++ b/Engine/Plugins/VirtualProduction/DMX/DMXGDTF/Source/DMXGDTF/Private/GDTF/Geometries/DMXGDTFStructureGeometry.cpp
@@ -6,6 +6,7 @@
 #include "GDTF/Geometries/DMXGDTFGeometry.h"
 #include "GDTF/Geometries/DMXGDTFGeometryCollect.h"
 #include "Serialization/DMXGDTFNodeInitializer.h"
+#include "Serialization/DMXGDTFXmlNodeBuilder.h"
 
 namespace UE::DMX::GDTF
 {
@@ -22,6 +23,21 @@ namespace UE::DMX::GDTF
 			.GetAttribute(TEXT("LinkedGeometry"), LinkedGeometry);
 	}
 
+	FXmlNode* FDMXGDTFStructureGeometry::CreateXmlNode(FXmlNode& Parent)
+	{
+		FXmlNode* AppendToNode = FDMXGDTFGeometry::CreateXmlNode(Parent);
+
+		const FDMXGDTFXmlNodeBuilder ChildBuilder = FDMXGDTFXmlNodeBuilder(Parent, *this, AppendToNode)
+			.SetAttribute(TEXT("StructureType"), StructureType)
+			.SetAttribute(TEXT("CrossSectionType"), CrossSectionType)
+			.SetAttribute(TEXT("CrossSectionHeight"), CrossSectionHeight)
+			.SetAttribute(TEXT("CrossSectionWallThickness"), CrossSectionWallThickness)
+			.SetAttribute(TEXT("TrussCrossSection"), TrussCrossSection)
+			.SetAttribute(TEXT("LinkedGeometry"), LinkedGeometry);
+
+		return ChildBuilder.GetIntermediateXmlNode();
+	}
+
 	TSharedPtr<FDMXGDTFGeometry> FDMXGDTFStructureGeometry::ResolveLinkedGeometry() const
 	{
 		const TSharedPtr<FDMXGDTFFixtureType> FixtureType = GetFixtureType().Pin();
diff --git a/Engine/Plugins/VirtualProduction/DMX/DMXGDTF/Source/DMXGDTF/Private/GDTF/Geometries/DMXGDTFSupportGeometry.cpp b/Engine/Plugins/VirtualProduction/DMX/DMXGDTF/Source/DMXGDTF/Private/GDTF/Geometries/DMXGDTFSupportGeometry.cpp
--- a/Engine/Plugins/VirtualProduction/DMX/DMXGDTF/Source/DMXGDTF/Private/GDTF/Geometries/DMXGDTFSupportGeometry.cpp
+++ b/Engine/Plugins/VirtualProduction/DMX/DMXGDTF/Source/DMXGDTF/Private/GDTF/Geometries/DMXGDTFSupportGeometry.cpp
@@ -3,6 +3,7 @@
 #include "GDTF/Geometries/DMXGDTFSupportGeometry.h"
 
 #include "Serialization/DMXGDTFNodeInitializer.h"
+#include "Serialization/DMXGDTFXmlNodeBuilder.h"
 
 namespace UE::DMX::GDTF
 {
@@ -27,4 +28,28 @@ namespace UE::DMX::GDTF
 			.GetAttribute(TEXT("ResistanceYY"), ResistanceYY)
 			.GetAttribute(TEXT("ResistanceZZ"), ResistanceZZ);
 	}
+
+	FXmlNode* FDMXGDTFSupportGeometry::CreateXmlNode(FXmlNode& Parent)
+	{
+		FXmlNode* AppendToNode = FDMXGDTFGeometry::CreateXmlNode(Parent);
+
+		const FDMXGDTFXmlNodeBuilder ChildBuilder = FDMXGDTFXmlNodeBuilder(Parent, *this, AppendToNode)
+			.SetAttribute(TEXT("SupportType"), SupportType)
+			.SetAttribute(TEXT("RopeCrossSection"), RopeCrossSection)
+			.SetAttribute(TEXT("RopeOffset"), RopeOffset)
+			.SetAttribute(TEXT("CapacityX"), CapacityX)
+			.SetAttribute(TEXT("CapacityY"), CapacityY)
+			.SetAttribute(TEXT("CapacityZ"), CapacityZ)
+			.SetAttribute(TEXT("CapacityXX"), CapacityXX)
+			.SetAttribute(TEXT("CapacityYY"), CapacityYY)
+			.SetAttribute(TEXT("CapacityZZ"), CapacityZZ)
+			.SetAttribute(TEXT("ResistanceX"), ResistanceX)
+			.SetAttribute(TEXT("ResistanceY"), ResistanceY)
+			.SetAttribute(TEXT("ResistanceZ"), ResistanceZ)
+			.SetAttribute(TEXT("ResistanceXX"), ResistanceXX)
+			.SetAttribute(TEXT("ResistanceYY"), ResistanceYY)
+			.SetAttribute(TEXT("ResistanceZZ"), ResistanceZZ);
+
+		return ChildBuilder.GetIntermediateXmlNode();
+	}
 }
diff --git a/Engine/Plugins/VirtualProduction/DMX/DMXGDTF/Source/DMXGDTF/Private/GDTF/Geometries/DMXGDTFWiringObjectGeometry.cpp b/Engine/Plugins/VirtualProduction/DMX/DMXGDTF/Source/DMXGDTF/Private/GDTF/Geometries/DMXGDTFWiringObjectGeometry.cpp
--- a/Engine/Plugins/VirtualProduction/DMX/DMXGDTF/Source/DMXGDTF/Private/GDTF/Geometries/DMXGDTFWiringObjectGeometry.cpp
+++ b/Engine/Plugins/VirtualProduction/DMX/DMXGDTF/Source/DMXGDTF/Private/GDTF/Geometries/DMXGDTFWiringObjectGeometry.cpp
@@ -6,6 +6,7 @@
 #include "GDTF/Geometries/DMXGDTFWiringObjectPinPatch.h"
 #include "GDTF/Models/DMXGDTFModel.h"
 #include "Serialization/DMXGDTFNodeInitializer.h"
+#include "Serialization/DMXGDTFXmlNodeBuilder.h"
 
 namespace UE::DMX::GDTF
 {
@@ -35,6 +36,36 @@ namespace UE::DMX::GDTF
 			.CreateChildren(TEXT("PinPatch"), PinPatchArray);
 	}
 
+	FXmlNode* FDMXGDTFWiringObjectGeometry::CreateXmlNode(FXmlNode& Parent)
+	{
+		const FTransform DefaultMatrix = FTransform::Identity;
+
+		FXmlNode* AppendToNode = FDMXGDTFGeometry::CreateXmlNode(Parent);
+
+		const FDMXGDTFXmlNodeBuilder ChildBuilder = FDMXGDTFXmlNodeBuilder(Parent, *this, AppendToNode)
+			.SetAttribute(TEXT("ConnectorType"), ConnectorType)
+			.SetAttribute(TEXT("Matrix"), Matrix, EDMXGDTFMatrixType::Matrix4x4, DefaultMatrix)
+			.SetAttribute(TEXT("ComponentType"), ComponentType)
+			.SetAttribute(TEXT("SignalType"), SignalType)
+			.SetAttribute(TEXT("PinCount"), PinCount)
+			.SetAttribute(TEXT("ElectricalPayLoad"), ElectricalPayLoad)
+			.SetAttribute(TEXT("VoltageRangeMax"), VoltageRangeMax)
+			.SetAttribute(TEXT("VoltageRangeMin"), VoltageRangeMin)
+			.SetAttribute(TEXT("FrequencyRangeMax"), FrequencyRangeMax)
+			.SetAttribute(TEXT("FrequencyRangeMin"), FrequencyRangeMin)
+			.SetAttribute(TEXT("MaxPayLoad"), MaxPayLoad)
+			.SetAttribute(TEXT("Voltage"), Voltage)
+			.SetAttribute(TEXT("SignalLayer"), SignalLayer)
+			.SetAttribute(TEXT("CosPhi"), CosPhi)
+			.SetAttribute(TEXT("FuseCurrent"), FuseCurrent)
+			.SetAttribute(TEXT("FuseRating"), FuseRating)
+			.SetAttribute(TEXT("Orientation"), Orientation)
+			.SetAttribute(TEXT("WireGroup"), WireGroup)
+			.AppendChildren(TEXT("PinPatch"), PinPatchArray);
+		
+		return ChildBuilder.GetIntermediateXmlNode();
+	}
+
 	TSharedPtr<FDMXGDTFModel> FDMXGDTFWiringObjectGeometry::ResolveModel() const
 	{
 		if (const TSharedPtr<FDMXGDTFFixtureType> FixtureType = GetFixtureType().Pin())
diff --git a/Engine/Plugins/VirtualProduction/DMX/DMXGDTF/Source/DMXGDTF/Private/GDTF/Geometries/DMXGDTFWiringObjectPinPatch.cpp b/Engine/Plugins/VirtualProduction/DMX/DMXGDTF/Source/DMXGDTF/Private/GDTF/Geometries/DMXGDTFWiringObjectPinPatch.cpp
--- a/Engine/Plugins/VirtualProduction/DMX/DMXGDTF/Source/DMXGDTF/Private/GDTF/Geometries/DMXGDTFWiringObjectPinPatch.cpp
+++ b/Engine/Plugins/VirtualProduction/DMX/DMXGDTF/Source/DMXGDTF/Private/GDTF/Geometries/DMXGDTFWiringObjectPinPatch.cpp
@@ -5,6 +5,7 @@
 #include "GDTF/Geometries/DMXGDTFGeometryCollect.h"
 #include "GDTF/Geometries/DMXGDTFWiringObjectGeometry.h"
 #include "Serialization/DMXGDTFNodeInitializer.h"
+#include "Serialization/DMXGDTFXmlNodeBuilder.h"
 
 namespace UE::DMX::GDTF
 {
@@ -20,6 +21,16 @@ namespace UE::DMX::GDTF
 			.GetAttribute(TEXT("ToWiringObject"), ToWiringObject);
 	}
 
+	FXmlNode* FDMXGDTFWiringObjectPinPatch::CreateXmlNode(FXmlNode& Parent)
+	{
+		const FDMXGDTFXmlNodeBuilder ChildBuilder = FDMXGDTFXmlNodeBuilder(Parent, *this)
+			.SetAttribute(TEXT("FromPin"), FromPin)
+			.SetAttribute(TEXT("ToPin"), ToPin)
+			.SetAttribute(TEXT("ToWiringObject"), ToWiringObject);
+
+		return ChildBuilder.GetIntermediateXmlNode();
+	}
+
 	TSharedPtr<FDMXGDTFWiringObjectGeometry> FDMXGDTFWiringObjectPinPatch::ResolveToWiringObject() const
 	{
 		const TSharedPtr<FDMXGDTFFixtureType> FixtureType = GetFixtureType().Pin();
diff --git a/Engine/Plugins/VirtualProduction/DMX/DMXGDTF/Source/DMXGDTF/Private/GDTF/Models/DMXGDTFModel.cpp b/Engine/Plugins/VirtualProduction/DMX/DMXGDTF/Source/DMXGDTF/Private/GDTF/Models/DMXGDTFModel.cpp
--- a/Engine/Plugins/VirtualProduction/DMX/DMXGDTF/Source/DMXGDTF/Private/GDTF/Models/DMXGDTFModel.cpp
+++ b/Engine/Plugins/VirtualProduction/DMX/DMXGDTF/Source/DMXGDTF/Private/GDTF/Models/DMXGDTFModel.cpp
@@ -3,6 +3,7 @@
 #include "GDTF/Models/DMXGDTFModel.h"
 
 #include "Serialization/DMXGDTFNodeInitializer.h"
+#include "Serialization/DMXGDTFXmlNodeBuilder.h"
 
 namespace UE::DMX::GDTF
 {
@@ -14,6 +15,7 @@ namespace UE::DMX::GDTF
 	{
 		FDMXGDTFNodeInitializer(SharedThis(this), XmlNode)
 			.GetAttribute(TEXT("Name"), Name)
+			.GetAttribute(TEXT("Length"), Length)
 			.GetAttribute(TEXT("Width"), Width)
 			.GetAttribute(TEXT("Height"), Height)
 			.GetAttribute(TEXT("PrimitiveType"), PrimitiveType)
@@ -25,4 +27,25 @@ namespace UE::DMX::GDTF
 			.GetAttribute(TEXT("SVGFrontOffsetX"), SVGFrontOffsetX)
 			.GetAttribute(TEXT("SVGFrontOffsetY"), SVGFrontOffsetY);
 	}
+
+	FXmlNode* FDMXGDTFModel::CreateXmlNode(FXmlNode& Parent)
+	{
+		const float DefaultSVGOffsetValue = 0.f;
+
+		const FDMXGDTFXmlNodeBuilder ChildBuilder = FDMXGDTFXmlNodeBuilder(Parent, *this)
+			.SetAttribute(TEXT("Name"), Name)
+			.SetAttribute(TEXT("Length"), Length)
+			.SetAttribute(TEXT("Width"), Width)
+			.SetAttribute(TEXT("Height"), Height)
+			.SetAttribute(TEXT("PrimitiveType"), PrimitiveType)
+			.SetAttribute(TEXT("File"), File)
+			.SetAttribute(TEXT("SVGOffsetX"), SVGOffsetX, DefaultSVGOffsetValue)
+			.SetAttribute(TEXT("SVGOffsetY"), SVGOffsetY, DefaultSVGOffsetValue)
+			.SetAttribute(TEXT("SVGSideOffsetX"), SVGSideOffsetX, DefaultSVGOffsetValue)
+			.SetAttribute(TEXT("SVGSideOffsetY"), SVGSideOffsetY, DefaultSVGOffsetValue)
+			.SetAttribute(TEXT("SVGFrontOffsetX"), SVGFrontOffsetX, DefaultSVGOffsetValue)
+			.SetAttribute(TEXT("SVGFrontOffsetY"), SVGFrontOffsetY, DefaultSVGOffsetValue);
+
+		return ChildBuilder.GetIntermediateXmlNode();
+	}
 }
diff --git a/Engine/Plugins/VirtualProduction/DMX/DMXGDTF/Source/DMXGDTF/Private/GDTF/PhysicalDescriptions/DMXGDTFColorRenderingIndex.cpp b/Engine/Plugins/VirtualProduction/DMX/DMXGDTF/Source/DMXGDTF/Private/GDTF/PhysicalDescriptions/DMXGDTFColorRenderingIndex.cpp
--- a/Engine/Plugins/VirtualProduction/DMX/DMXGDTF/Source/DMXGDTF/Private/GDTF/PhysicalDescriptions/DMXGDTFColorRenderingIndex.cpp
+++ b/Engine/Plugins/VirtualProduction/DMX/DMXGDTF/Source/DMXGDTF/Private/GDTF/PhysicalDescriptions/DMXGDTFColorRenderingIndex.cpp
@@ -3,6 +3,7 @@
 #include "GDTF/PhysicalDescriptions/DMXGDTFColorRenderingIndex.h"
 
 #include "Serialization/DMXGDTFNodeInitializer.h"
+#include "Serialization/DMXGDTFXmlNodeBuilder.h"
 
 namespace UE::DMX::GDTF
 {
@@ -17,6 +18,15 @@ namespace UE::DMX::GDTF
 			.GetAttribute(TEXT("ColorRenderingIndex"), ColorRenderingIndex);
 	}
 
+	FXmlNode* FDMXGDTFColorRenderingIndex::CreateXmlNode(FXmlNode& Parent)
+	{
+		const FDMXGDTFXmlNodeBuilder ChildBuilder = FDMXGDTFXmlNodeBuilder(Parent, *this)
+			.SetAttribute(TEXT("CES"), CES)
+			.SetAttribute(TEXT("ColorRenderingIndex"), ColorRenderingIndex);
+
+		return ChildBuilder.GetIntermediateXmlNode();
+	}
+
 	uint8 FDMXGDTFColorRenderingIndex::ParseCES(const FString& GDTFString) const
 	{
 		const FString CleanString = GDTFString.Replace(TEXT("CES"), TEXT(""));
diff --git a/Engine/Plugins/VirtualProduction/DMX/DMXGDTF/Source/DMXGDTF/Private/GDTF/PhysicalDescriptions/DMXGDTFColorRenderingIndexGroup.cpp b/Engine/Plugins/VirtualProduction/DMX/DMXGDTF/Source/DMXGDTF/Private/GDTF/PhysicalDescriptions/DMXGDTFColorRenderingIndexGroup.cpp
--- a/Engine/Plugins/VirtualProduction/DMX/DMXGDTF/Source/DMXGDTF/Private/GDTF/PhysicalDescriptions/DMXGDTFColorRenderingIndexGroup.cpp
+++ b/Engine/Plugins/VirtualProduction/DMX/DMXGDTF/Source/DMXGDTF/Private/GDTF/PhysicalDescriptions/DMXGDTFColorRenderingIndexGroup.cpp
@@ -4,6 +4,7 @@
 
 #include "GDTF/PhysicalDescriptions/DMXGDTFColorRenderingIndex.h"
 #include "Serialization/DMXGDTFNodeInitializer.h"
+#include "Serialization/DMXGDTFXmlNodeBuilder.h"
 
 namespace UE::DMX::GDTF
 {
@@ -17,4 +18,13 @@ namespace UE::DMX::GDTF
 			.GetAttribute(TEXT("ColorTemperature"), ColorTemperature)
 			.CreateChildren(TEXT("CRI"), CRIArray);
 	}
+
+	FXmlNode* FDMXGDTFColorRenderingIndexGroup::CreateXmlNode(FXmlNode& Parent)
+	{
+		const FDMXGDTFXmlNodeBuilder ChildBuilder = FDMXGDTFXmlNodeBuilder(Parent, *this)
+			.SetAttribute(TEXT("ColorTemperature"), ColorTemperature)
+			.AppendChildren(TEXT("CRI"), CRIArray);
+
+		return ChildBuilder.GetIntermediateXmlNode();
+	}
 }
diff --git a/Engine/Plugins/VirtualProduction/DMX/DMXGDTF/Source/DMXGDTF/Private/GDTF/PhysicalDescriptions/DMXGDTFColorSpace.cpp b/Engine/Plugins/VirtualProduction/DMX/DMXGDTF/Source/DMXGDTF/Private/GDTF/PhysicalDescriptions/DMXGDTFColorSpace.cpp
--- a/Engine/Plugins/VirtualProduction/DMX/DMXGDTF/Source/DMXGDTF/Private/GDTF/PhysicalDescriptions/DMXGDTFColorSpace.cpp
+++ b/Engine/Plugins/VirtualProduction/DMX/DMXGDTF/Source/DMXGDTF/Private/GDTF/PhysicalDescriptions/DMXGDTFColorSpace.cpp
@@ -3,6 +3,7 @@
 #include "GDTF/PhysicalDescriptions/DMXGDTFColorSpace.h"
 
 #include "Serialization/DMXGDTFNodeInitializer.h"
+#include "Serialization/DMXGDTFXmlNodeBuilder.h"
 
 namespace UE::DMX::GDTF
 {
@@ -20,4 +21,20 @@ namespace UE::DMX::GDTF
 			.GetAttribute(TEXT("Blue"), Blue)
 			.GetAttribute(TEXT("WhitePoint"), WhitePoint);
 	}
+
+	FXmlNode* FDMXGDTFColorSpace::CreateXmlNode(FXmlNode& Parent)
+	{
+		const FName DefaultName = "Default";
+		const FDMXGDTFColorCIE1931xyY DefaultColor = { 0.f, 0.f, 0.f };
+
+		const FDMXGDTFXmlNodeBuilder ChildBuilder = FDMXGDTFXmlNodeBuilder(Parent, *this)
+			.SetAttribute(TEXT("Name"), Name, DefaultName)
+			.SetAttribute(TEXT("Mode"), Mode)
+			.SetAttribute(TEXT("Red"), Red, DefaultColor)
+			.SetAttribute(TEXT("Green"), Green, DefaultColor)
+			.SetAttribute(TEXT("Blue"), Blue, DefaultColor)
+			.SetAttribute(TEXT("WhitePoint"), WhitePoint, DefaultColor);
+
+		return ChildBuilder.GetIntermediateXmlNode();
+	}
 }
diff --git a/Engine/Plugins/VirtualProduction/DMX/DMXGDTF/Source/DMXGDTF/Private/GDTF/PhysicalDescriptions/DMXGDTFEmitter.cpp b/Engine/Plugins/VirtualProduction/DMX/DMXGDTF/Source/DMXGDTF/Private/GDTF/PhysicalDescriptions/DMXGDTFEmitter.cpp
--- a/Engine/Plugins/VirtualProduction/DMX/DMXGDTF/Source/DMXGDTF/Private/GDTF/PhysicalDescriptions/DMXGDTFEmitter.cpp
+++ b/Engine/Plugins/VirtualProduction/DMX/DMXGDTF/Source/DMXGDTF/Private/GDTF/PhysicalDescriptions/DMXGDTFEmitter.cpp
@@ -4,6 +4,7 @@
 
 #include "GDTF/PhysicalDescriptions/DMXGDTFMeasurement.h"
 #include "Serialization/DMXGDTFNodeInitializer.h"
+#include "Serialization/DMXGDTFXmlNodeBuilder.h"
 
 namespace UE::DMX::GDTF
 {
@@ -15,9 +16,21 @@ namespace UE::DMX::GDTF
 	{
 		FDMXGDTFNodeInitializer(SharedThis(this), XmlNode)
 			.GetAttribute(TEXT("Name"), Name)
-			.GetAttribute(TEXT("ColorCIE"), ColorCIE)
+			.GetAttribute(TEXT("Color"), Color)
 			.GetAttribute(TEXT("DominantWaveLength"), DominantWaveLength)
 			.GetAttribute(TEXT("DiodePart"), DiodePart)
 			.CreateChildren(TEXT("Measurement"), Measurements);
 	}
+
+	FXmlNode* FDMXGDTFEmitter::CreateXmlNode(FXmlNode& Parent)
+	{
+		const FDMXGDTFXmlNodeBuilder ChildBuilder = FDMXGDTFXmlNodeBuilder(Parent, *this)
+			.SetAttribute(TEXT("Name"), Name)
+			.SetAttribute(TEXT("Color"), Color)
+			.SetAttribute(TEXT("DominantWaveLength"), DominantWaveLength)
+			.SetAttribute(TEXT("DiodePart"), DiodePart)
+			.AppendChildren(TEXT("Measurement"), Measurements);
+
+		return ChildBuilder.GetIntermediateXmlNode();
+	}
 }
diff --git a/Engine/Plugins/VirtualProduction/DMX/DMXGDTF/Source/DMXGDTF/Private/GDTF/PhysicalDescriptions/DMXGDTFFilter.cpp b/Engine/Plugins/VirtualProduction/DMX/DMXGDTF/Source/DMXGDTF/Private/GDTF/PhysicalDescriptions/DMXGDTFFilter.cpp
--- a/Engine/Plugins/VirtualProduction/DMX/DMXGDTF/Source/DMXGDTF/Private/GDTF/PhysicalDescriptions/DMXGDTFFilter.cpp
+++ b/Engine/Plugins/VirtualProduction/DMX/DMXGDTF/Source/DMXGDTF/Private/GDTF/PhysicalDescriptions/DMXGDTFFilter.cpp
@@ -4,6 +4,7 @@
 
 #include "GDTF/PhysicalDescriptions/DMXGDTFMeasurement.h"
 #include "Serialization/DMXGDTFNodeInitializer.h"
+#include "Serialization/DMXGDTFXmlNodeBuilder.h"
 
 namespace UE::DMX::GDTF
 {
@@ -15,7 +16,17 @@ namespace UE::DMX::GDTF
 	{
 		FDMXGDTFNodeInitializer(SharedThis(this), XmlNode)
 			.GetAttribute(TEXT("Name"), Name)
-			.GetAttribute(TEXT("ColorCIE"), ColorCIE)
+			.GetAttribute(TEXT("Color"), Color)
 			.CreateChildren(TEXT("Measurement"), Measurements);
 	}
+
+	FXmlNode* FDMXGDTFFilter::CreateXmlNode(FXmlNode& Parent)
+	{
+		const FDMXGDTFXmlNodeBuilder ChildBuilder = FDMXGDTFXmlNodeBuilder(Parent, *this)
+			.SetAttribute(TEXT("Name"), Name)
+			.SetAttribute(TEXT("Color"), Color)
+			.AppendChildren(TEXT("Measurement"), Measurements);
+
+		return ChildBuilder.GetIntermediateXmlNode();
+	}
 }
diff --git a/Engine/Plugins/VirtualProduction/DMX/DMXGDTF/Source/DMXGDTF/Private/GDTF/PhysicalDescriptions/DMXGDTFGamut.cpp b/Engine/Plugins/VirtualProduction/DMX/DMXGDTF/Source/DMXGDTF/Private/GDTF/PhysicalDescriptions/DMXGDTFGamut.cpp
--- a/Engine/Plugins/VirtualProduction/DMX/DMXGDTF/Source/DMXGDTF/Private/GDTF/PhysicalDescriptions/DMXGDTFGamut.cpp
+++ b/Engine/Plugins/VirtualProduction/DMX/DMXGDTF/Source/DMXGDTF/Private/GDTF/PhysicalDescriptions/DMXGDTFGamut.cpp
@@ -4,6 +4,7 @@
 
 #include "DMXGDTFLog.h"
 #include "Serialization/DMXGDTFNodeInitializer.h"
+#include "Serialization/DMXGDTFXmlNodeBuilder.h"
 
 namespace UE::DMX::GDTF
 {
@@ -18,14 +19,23 @@ namespace UE::DMX::GDTF
 			.GetAttribute(TEXT("Points"), Points, this, &FDMXGDTFGamut::ParsePoints);
 	}
 
-	TArray<FVector> FDMXGDTFGamut::ParsePoints(const FString& GDTFString) const
+	FXmlNode* FDMXGDTFGamut::CreateXmlNode(FXmlNode& Parent)
+	{
+		const FDMXGDTFXmlNodeBuilder ChildBuilder = FDMXGDTFXmlNodeBuilder(Parent, *this)
+			.SetAttribute(TEXT("Name"), Name)
+			.SetAttribute(TEXT("Points"), Points);
+
+		return ChildBuilder.GetIntermediateXmlNode();
+	}
+
+	TArray<FDMXGDTFColorCIE1931xyY> FDMXGDTFGamut::ParsePoints(const FString& GDTFString) const
 	{
 		const FString CleanString = GDTFString.Replace(TEXT("}"), TEXT(""));
 
 		TArray<FString> Strings;
 		GDTFString.ParseIntoArray(Strings, TEXT("{"));
 
-		TArray<FVector> Result;
+		TArray<FDMXGDTFColorCIE1931xyY> Result;
 		for (const FString& String : Strings)
 		{
 			TArray<FString> ComponentStrings;
@@ -39,10 +49,10 @@ namespace UE::DMX::GDTF
 				break;
 			}
 
-			FVector Point;
+			FDMXGDTFColorCIE1931xyY Point;
 			if (LexTryParseString(Point.X, *ComponentStrings[0]) &&
 				LexTryParseString(Point.Y, *ComponentStrings[1]) &&
-				LexTryParseString(Point.Z, *ComponentStrings[2]))
+				LexTryParseString(Point.YY, *ComponentStrings[2]))
 			{
 				Result.Add(Point);
 			}
diff --git a/Engine/Plugins/VirtualProduction/DMX/DMXGDTF/Source/DMXGDTF/Private/GDTF/PhysicalDescriptions/DMXGDTFLegHeight.cpp b/Engine/Plugins/VirtualProduction/DMX/DMXGDTF/Source/DMXGDTF/Private/GDTF/PhysicalDescriptions/DMXGDTFLegHeight.cpp
--- a/Engine/Plugins/VirtualProduction/DMX/DMXGDTF/Source/DMXGDTF/Private/GDTF/PhysicalDescriptions/DMXGDTFLegHeight.cpp
+++ b/Engine/Plugins/VirtualProduction/DMX/DMXGDTF/Source/DMXGDTF/Private/GDTF/PhysicalDescriptions/DMXGDTFLegHeight.cpp
@@ -3,6 +3,7 @@
 #include "GDTF/PhysicalDescriptions/DMXGDTFLegHeight.h"
 
 #include "Serialization/DMXGDTFNodeInitializer.h"
+#include "Serialization/DMXGDTFXmlNodeBuilder.h"
 
 namespace UE::DMX::GDTF
 {
@@ -15,4 +16,12 @@ namespace UE::DMX::GDTF
 		FDMXGDTFNodeInitializer(SharedThis(this), XmlNode)
 			.GetAttribute(TEXT("Value"), Value);
 	}
+
+	FXmlNode* FDMXGDTFLegHeight::CreateXmlNode(FXmlNode& Parent)
+	{
+		const FDMXGDTFXmlNodeBuilder ChildBuilder = FDMXGDTFXmlNodeBuilder(Parent, *this)
+			.SetAttribute(TEXT("Value"), Value);
+
+		return ChildBuilder.GetIntermediateXmlNode();
+	}
 }
diff --git a/Engine/Plugins/VirtualProduction/DMX/DMXGDTF/Source/DMXGDTF/Private/GDTF/PhysicalDescriptions/DMXGDTFMeasurement.cpp b/Engine/Plugins/VirtualProduction/DMX/DMXGDTF/Source/DMXGDTF/Private/GDTF/PhysicalDescriptions/DMXGDTFMeasurement.cpp
--- a/Engine/Plugins/VirtualProduction/DMX/DMXGDTF/Source/DMXGDTF/Private/GDTF/PhysicalDescriptions/DMXGDTFMeasurement.cpp
+++ b/Engine/Plugins/VirtualProduction/DMX/DMXGDTF/Source/DMXGDTF/Private/GDTF/PhysicalDescriptions/DMXGDTFMeasurement.cpp
@@ -4,6 +4,7 @@
 
 #include "GDTF/PhysicalDescriptions/DMXGDTFMeasurementPoint.h"
 #include "Serialization/DMXGDTFNodeInitializer.h"
+#include "Serialization/DMXGDTFXmlNodeBuilder.h"
 
 namespace UE::DMX::GDTF
 {
@@ -14,10 +15,26 @@ namespace UE::DMX::GDTF
 			.GetAttribute(TEXT("LuminousIntensity"), LuminousIntensity)
 			.GetAttribute(TEXT("Transmission"), Transmission)
 			.GetAttribute(TEXT("InterpolationTo"), InterpolationTo)
-			.CreateChildren(TEXT("MeasurementPoints"), MeasurementPointArray);
+			.CreateChildren(TEXT("MeasurementPoint"), MeasurementPointArray);
 	}
 
-	FDMXGDTFEmitterMeasurement::FDMXGDTFEmitterMeasurement(const TWeakPtr<FDMXGDTFEmitter>& InEmitter)
+	FXmlNode* FDMXGDTFMeasurementBase::CreateXmlNode(FXmlNode& Parent)
+	{
+		const float DefaultPhysical = 0.f;
+		const float DefaultLuminousIntensity = 0.f;
+		const float DefaultTransmission = 0.f;
+
+		const FDMXGDTFXmlNodeBuilder ChildBuilder = FDMXGDTFXmlNodeBuilder(Parent, *this)
+			.SetAttribute(TEXT("Physical"), Physical, DefaultPhysical)
+			.SetAttribute(TEXT("LuminousIntensity"), LuminousIntensity, DefaultLuminousIntensity)
+			.SetAttribute(TEXT("Transmission"), Transmission, DefaultTransmission)
+			.SetAttribute(TEXT("InterpolationTo"), InterpolationTo)
+			.AppendChildren(TEXT("MeasurementPoint"), MeasurementPointArray);
+
+		return ChildBuilder.GetIntermediateXmlNode();
+	}
+
+	FDMXGDTFEmitterMeasurement::FDMXGDTFEmitterMeasurement(const TSharedRef<FDMXGDTFEmitter>& InEmitter)
 		: OuterEmitter(InEmitter)
 	{}
 
diff --git a/Engine/Plugins/VirtualProduction/DMX/DMXGDTF/Source/DMXGDTF/Private/GDTF/PhysicalDescriptions/DMXGDTFMeasurementPoint.cpp b/Engine/Plugins/VirtualProduction/DMX/DMXGDTF/Source/DMXGDTF/Private/GDTF/PhysicalDescriptions/DMXGDTFMeasurementPoint.cpp
--- a/Engine/Plugins/VirtualProduction/DMX/DMXGDTF/Source/DMXGDTF/Private/GDTF/PhysicalDescriptions/DMXGDTFMeasurementPoint.cpp
+++ b/Engine/Plugins/VirtualProduction/DMX/DMXGDTF/Source/DMXGDTF/Private/GDTF/PhysicalDescriptions/DMXGDTFMeasurementPoint.cpp
@@ -3,6 +3,7 @@
 #include "GDTF/PhysicalDescriptions/DMXGDTFMeasurementPoint.h"
 
 #include "Serialization/DMXGDTFNodeInitializer.h"
+#include "Serialization/DMXGDTFXmlNodeBuilder.h"
 
 namespace UE::DMX::GDTF
 {
@@ -16,4 +17,13 @@ namespace UE::DMX::GDTF
 			.GetAttribute(TEXT("WaveLength"), WaveLength)
 			.GetAttribute(TEXT("Energy"), Energy);
 	}
+
+	FXmlNode* FDMXGDTFMeasurementPoint::CreateXmlNode(FXmlNode& Parent)
+	{
+		const FDMXGDTFXmlNodeBuilder ChildBuilder = FDMXGDTFXmlNodeBuilder(Parent, *this)
+			.SetAttribute(TEXT("WaveLength"), WaveLength)
+			.SetAttribute(TEXT("Energy"), Energy);
+
+		return ChildBuilder.GetIntermediateXmlNode();
+	}
 }
diff --git a/Engine/Plugins/VirtualProduction/DMX/DMXGDTF/Source/DMXGDTF/Private/GDTF/PhysicalDescriptions/DMXGDTFOperatingTemperature.cpp b/Engine/Plugins/VirtualProduction/DMX/DMXGDTF/Source/DMXGDTF/Private/GDTF/PhysicalDescriptions/DMXGDTFOperatingTemperature.cpp
--- a/Engine/Plugins/VirtualProduction/DMX/DMXGDTF/Source/DMXGDTF/Private/GDTF/PhysicalDescriptions/DMXGDTFOperatingTemperature.cpp
+++ b/Engine/Plugins/VirtualProduction/DMX/DMXGDTF/Source/DMXGDTF/Private/GDTF/PhysicalDescriptions/DMXGDTFOperatingTemperature.cpp
@@ -3,6 +3,7 @@
 #include "GDTF/PhysicalDescriptions/DMXGDTFOperatingTemperature.h"
 
 #include "Serialization/DMXGDTFNodeInitializer.h"
+#include "Serialization/DMXGDTFXmlNodeBuilder.h"
 
 namespace UE::DMX::GDTF
 {
@@ -16,4 +17,13 @@ namespace UE::DMX::GDTF
 			.GetAttribute(TEXT("Low"), Low)
 			.GetAttribute(TEXT("High"), High);
 	}
+
+	FXmlNode* FDMXGDTFOperatingTemperature::CreateXmlNode(FXmlNode& Parent)
+	{
+		const FDMXGDTFXmlNodeBuilder ChildBuilder = FDMXGDTFXmlNodeBuilder(Parent, *this)
+			.SetAttribute(TEXT("Low"), Low)
+			.SetAttribute(TEXT("High"), High);
+
+		return ChildBuilder.GetIntermediateXmlNode();
+	}
 }
diff --git a/Engine/Plugins/VirtualProduction/DMX/DMXGDTF/Source/DMXGDTF/Private/GDTF/PhysicalDescriptions/DMXGDTFPhysicalDescriptions.cpp b/Engine/Plugins/VirtualProduction/DMX/DMXGDTF/Source/DMXGDTF/Private/GDTF/PhysicalDescriptions/DMXGDTFPhysicalDescriptions.cpp
--- a/Engine/Plugins/VirtualProduction/DMX/DMXGDTF/Source/DMXGDTF/Private/GDTF/PhysicalDescriptions/DMXGDTFPhysicalDescriptions.cpp
+++ b/Engine/Plugins/VirtualProduction/DMX/DMXGDTF/Source/DMXGDTF/Private/GDTF/PhysicalDescriptions/DMXGDTFPhysicalDescriptions.cpp
@@ -10,6 +10,7 @@
 #include "GDTF/PhysicalDescriptions/DMXGDTFGamut.h"
 #include "GDTF/PhysicalDescriptions/DMXGDTFProperties.h"
 #include "Serialization/DMXGDTFNodeInitializer.h"
+#include "Serialization/DMXGDTFXmlNodeBuilder.h"
 
 namespace UE::DMX::GDTF
 {
@@ -22,11 +23,26 @@ namespace UE::DMX::GDTF
 		FDMXGDTFNodeInitializer(SharedThis(this), XmlNode)
 			.CreateChildCollection(TEXT("Emitters"), TEXT("Emitter"), Emitters)
 			.CreateChildCollection(TEXT("Filters"), TEXT("Filter"), Filters)
-			.CreateChildCollection(TEXT("ColorSpaces"), TEXT("ColorSpace"), ColorSpaces)
+			.CreateOptionalChild(TEXT("ColorSpace"), ColorSpaces)
 			.CreateChildCollection(TEXT("AdditionalColorSpaces"), TEXT("ColorSpace"), AdditionalColorSpaces)
-			.CreateChildCollection(TEXT("Gamut"), TEXT("Gamuts"), Gamuts)
+			.CreateChildCollection(TEXT("Gamuts"), TEXT("Gamut"), Gamuts)
 			.CreateChildCollection(TEXT("DMXProfiles"), TEXT("DMXProfile"), DMXProfiles)
 			.CreateChildCollection(TEXT("CRIs"), TEXT("CRIGroup"), CRIs)
 			.CreateOptionalChild(TEXT("Properties"), Properties);
 	}
+
+	FXmlNode* FDMXGDTFPhysicalDescriptions::CreateXmlNode(FXmlNode& Parent)
+	{
+		const FDMXGDTFXmlNodeBuilder ChildBuilder = FDMXGDTFXmlNodeBuilder(Parent, *this)
+			.AppendChildCollection(TEXT("Emitters"), TEXT("Emitter"), Emitters)
+			.AppendChildCollection(TEXT("Filters"), TEXT("Filter"), Filters)
+			.AppendOptionalChild(TEXT("ColorSpace"), ColorSpaces)
+			.AppendChildCollection(TEXT("AdditionalColorSpaces"), TEXT("AdditionalColorSpace"), AdditionalColorSpaces)
+			.AppendChildCollection(TEXT("Gamuts"), TEXT("Gamut"), Gamuts)
+			.AppendChildCollection(TEXT("DMXProfiles"), TEXT("DMXProfile"), DMXProfiles)
+			.AppendChildCollection(TEXT("CRIs"), TEXT("CRIGroup"), CRIs)
+			.AppendOptionalChild(TEXT("Properties"), Properties);
+
+		return ChildBuilder.GetIntermediateXmlNode();
+	}
 }
diff --git a/Engine/Plugins/VirtualProduction/DMX/DMXGDTF/Source/DMXGDTF/Private/GDTF/PhysicalDescriptions/DMXGDTFProfile.cpp b/Engine/Plugins/VirtualProduction/DMX/DMXGDTF/Source/DMXGDTF/Private/GDTF/PhysicalDescriptions/DMXGDTFProfile.cpp
--- a/Engine/Plugins/VirtualProduction/DMX/DMXGDTF/Source/DMXGDTF/Private/GDTF/PhysicalDescriptions/DMXGDTFProfile.cpp
+++ b/Engine/Plugins/VirtualProduction/DMX/DMXGDTF/Source/DMXGDTF/Private/GDTF/PhysicalDescriptions/DMXGDTFProfile.cpp
@@ -4,6 +4,7 @@
 
 #include "GDTF/PhysicalDescriptions/DMXGDTFDMXProfilePoint.h"
 #include "Serialization/DMXGDTFNodeInitializer.h"
+#include "Serialization/DMXGDTFXmlNodeBuilder.h"
 
 namespace UE::DMX::GDTF
 {
@@ -17,4 +18,13 @@ namespace UE::DMX::GDTF
 			.GetAttribute(TEXT("Name"), Name)
 			.CreateChildren(TEXT("Point"), PointArray);
 	}
+
+	FXmlNode* FDMXGDTFDMXProfile::CreateXmlNode(FXmlNode& Parent)
+	{
+		const FDMXGDTFXmlNodeBuilder ChildBuilder = FDMXGDTFXmlNodeBuilder(Parent, *this)
+			.SetAttribute(TEXT("Name"), Name)
+			.AppendChildren(TEXT("Point"), PointArray);
+
+		return ChildBuilder.GetIntermediateXmlNode();
+	}
 }
diff --git a/Engine/Plugins/VirtualProduction/DMX/DMXGDTF/Source/DMXGDTF/Private/GDTF/PhysicalDescriptions/DMXGDTFProfilePoint.cpp b/Engine/Plugins/VirtualProduction/DMX/DMXGDTF/Source/DMXGDTF/Private/GDTF/PhysicalDescriptions/DMXGDTFProfilePoint.cpp
--- a/Engine/Plugins/VirtualProduction/DMX/DMXGDTF/Source/DMXGDTF/Private/GDTF/PhysicalDescriptions/DMXGDTFProfilePoint.cpp
+++ b/Engine/Plugins/VirtualProduction/DMX/DMXGDTF/Source/DMXGDTF/Private/GDTF/PhysicalDescriptions/DMXGDTFProfilePoint.cpp
@@ -3,6 +3,7 @@
 #include "GDTF/PhysicalDescriptions/DMXGDTFDMXProfilePoint.h"
 
 #include "Serialization/DMXGDTFNodeInitializer.h"
+#include "Serialization/DMXGDTFXmlNodeBuilder.h"
 
 namespace UE::DMX::GDTF
 {
@@ -19,4 +20,16 @@ namespace UE::DMX::GDTF
 			.GetAttribute(TEXT("CFC2"), CFC2)
 			.GetAttribute(TEXT("CFC3"), CFC3);
 	}
+
+	FXmlNode* FDMXGDTFDMXProfilePoint::CreateXmlNode(FXmlNode& Parent)
+	{
+		const FDMXGDTFXmlNodeBuilder ChildBuilder = FDMXGDTFXmlNodeBuilder(Parent, *this)
+			.SetAttribute(TEXT("DMXPercentage"), DMXPercentage)
+			.SetAttribute(TEXT("CFC0"), CFC0)
+			.SetAttribute(TEXT("CFC1"), CFC1)
+			.SetAttribute(TEXT("CFC2"), CFC2)
+			.SetAttribute(TEXT("CFC3"), CFC3);
+
+		return ChildBuilder.GetIntermediateXmlNode();
+	}
 }
diff --git a/Engine/Plugins/VirtualProduction/DMX/DMXGDTF/Source/DMXGDTF/Private/GDTF/PhysicalDescriptions/DMXGDTFProperties.cpp b/Engine/Plugins/VirtualProduction/DMX/DMXGDTF/Source/DMXGDTF/Private/GDTF/PhysicalDescriptions/DMXGDTFProperties.cpp
--- a/Engine/Plugins/VirtualProduction/DMX/DMXGDTF/Source/DMXGDTF/Private/GDTF/PhysicalDescriptions/DMXGDTFProperties.cpp
+++ b/Engine/Plugins/VirtualProduction/DMX/DMXGDTF/Source/DMXGDTF/Private/GDTF/PhysicalDescriptions/DMXGDTFProperties.cpp
@@ -6,6 +6,7 @@
 #include "GDTF/PhysicalDescriptions/DMXGDTFOperatingTemperature.h"
 #include "GDTF/PhysicalDescriptions/DMXGDTFWeight.h"
 #include "Serialization/DMXGDTFNodeInitializer.h"
+#include "Serialization/DMXGDTFXmlNodeBuilder.h"
 
 namespace UE::DMX::GDTF
 {
@@ -17,7 +18,17 @@ namespace UE::DMX::GDTF
 	{
 		FDMXGDTFNodeInitializer(SharedThis(this), XmlNode)
 			.CreateOptionalChild(TEXT("OperatingTemperature"), OperatingTemperature)
-			.CreateOptionalChild(TEXT("Weigth"), Weigth)
+			.CreateOptionalChild(TEXT("Weight"), Weight)
 			.CreateOptionalChild(TEXT("LegHeight"), LegHeight);
 	}
+
+	FXmlNode* FDMXGDTFProperties::CreateXmlNode(FXmlNode& Parent)
+	{
+		const FDMXGDTFXmlNodeBuilder ChildBuilder = FDMXGDTFXmlNodeBuilder(Parent, *this)
+			.AppendOptionalChild(TEXT("OperatingTemperature"), OperatingTemperature)
+			.AppendOptionalChild(TEXT("Weight"), Weight)
+			.AppendOptionalChild(TEXT("LegHeight"), LegHeight);
+
+		return ChildBuilder.GetIntermediateXmlNode();
+	}
 }
diff --git a/Engine/Plugins/VirtualProduction/DMX/DMXGDTF/Source/DMXGDTF/Private/GDTF/PhysicalDescriptions/DMXGDTFWeight.cpp b/Engine/Plugins/VirtualProduction/DMX/DMXGDTF/Source/DMXGDTF/Private/GDTF/PhysicalDescriptions/DMXGDTFWeight.cpp
--- a/Engine/Plugins/VirtualProduction/DMX/DMXGDTF/Source/DMXGDTF/Private/GDTF/PhysicalDescriptions/DMXGDTFWeight.cpp
+++ b/Engine/Plugins/VirtualProduction/DMX/DMXGDTF/Source/DMXGDTF/Private/GDTF/PhysicalDescriptions/DMXGDTFWeight.cpp
@@ -3,6 +3,7 @@
 #include "GDTF/PhysicalDescriptions/DMXGDTFWeight.h"
 
 #include "Serialization/DMXGDTFNodeInitializer.h"
+#include "Serialization/DMXGDTFXmlNodeBuilder.h"
 
 namespace UE::DMX::GDTF
 {
@@ -15,4 +16,12 @@ namespace UE::DMX::GDTF
 		FDMXGDTFNodeInitializer(SharedThis(this), XmlNode)
 			.GetAttribute(TEXT("Value"), Value);
 	}
+
+	FXmlNode* FDMXGDTFWeight::CreateXmlNode(FXmlNode& Parent)
+	{
+		const FDMXGDTFXmlNodeBuilder ChildBuilder = FDMXGDTFXmlNodeBuilder(Parent, *this)
+			.SetAttribute(TEXT("Value"), Value);
+
+		return ChildBuilder.GetIntermediateXmlNode();
+	}
 }
diff --git a/Engine/Plugins/VirtualProduction/DMX/DMXGDTF/Source/DMXGDTF/Private/GDTF/Protocols/ArtNet/DMXGDTFProtocolArtNet.cpp b/Engine/Plugins/VirtualProduction/DMX/DMXGDTF/Source/DMXGDTF/Private/GDTF/Protocols/ArtNet/DMXGDTFProtocolArtNet.cpp
--- a/Engine/Plugins/VirtualProduction/DMX/DMXGDTF/Source/DMXGDTF/Private/GDTF/Protocols/ArtNet/DMXGDTFProtocolArtNet.cpp
+++ b/Engine/Plugins/VirtualProduction/DMX/DMXGDTF/Source/DMXGDTF/Private/GDTF/Protocols/ArtNet/DMXGDTFProtocolArtNet.cpp
@@ -4,6 +4,7 @@
 
 #include "GDTF/Protocols/DMXGDTFProtocolDMXMap.h"
 #include "Serialization/DMXGDTFNodeInitializer.h"
+#include "Serialization/DMXGDTFXmlNodeBuilder.h"
 
 namespace UE::DMX::GDTF::ArtNet
 {
@@ -16,4 +17,12 @@ namespace UE::DMX::GDTF::ArtNet
 		FDMXGDTFNodeInitializer(SharedThis(this), XmlNode)
 			.CreateChildren(TEXT("Map"), Maps);
 	}
+
+	FXmlNode* FDMXGDTFProtocolArtNet::CreateXmlNode(FXmlNode& Parent)
+	{
+		const FDMXGDTFXmlNodeBuilder ChildBuilder = FDMXGDTFXmlNodeBuilder(Parent, *this)
+			.AppendChildren(TEXT("Map"), Maps);
+
+		return ChildBuilder.GetIntermediateXmlNode();
+	}
 }
diff --git a/Engine/Plugins/VirtualProduction/DMX/DMXGDTF/Source/DMXGDTF/Private/GDTF/Protocols/DMXGDTFProtocolDMXMap.cpp b/Engine/Plugins/VirtualProduction/DMX/DMXGDTF/Source/DMXGDTF/Private/GDTF/Protocols/DMXGDTFProtocolDMXMap.cpp
--- a/Engine/Plugins/VirtualProduction/DMX/DMXGDTF/Source/DMXGDTF/Private/GDTF/Protocols/DMXGDTFProtocolDMXMap.cpp
+++ b/Engine/Plugins/VirtualProduction/DMX/DMXGDTF/Source/DMXGDTF/Private/GDTF/Protocols/DMXGDTFProtocolDMXMap.cpp
@@ -3,6 +3,7 @@
 #include "GDTF/Protocols/DMXGDTFProtocolDMXMap.h"
 
 #include "Serialization/DMXGDTFNodeInitializer.h"
+#include "Serialization/DMXGDTFXmlNodeBuilder.h"
 
 namespace UE::DMX::GDTF
 {
@@ -13,9 +14,18 @@ namespace UE::DMX::GDTF
 			.GetAttribute(TEXT("Value"), Value);
 	}
 
+	FXmlNode* FDMXGDTFProtocolDMXMapBase::CreateXmlNode(FXmlNode& Parent)
+	{
+		const FDMXGDTFXmlNodeBuilder ChildBuilder = FDMXGDTFXmlNodeBuilder(Parent, *this)
+			.SetAttribute(TEXT("Key"), Key)
+			.SetAttribute(TEXT("Value"), Value);
+
+		return ChildBuilder.GetIntermediateXmlNode();
+	}
+
 	namespace ArtNet
 	{
-		FDMXGDTFProtocolArtNetDMXMap::FDMXGDTFProtocolArtNetDMXMap(const TWeakPtr<ArtNet::FDMXGDTFProtocolArtNet>& InProtocolArtNet)
+		FDMXGDTFProtocolArtNetDMXMap::FDMXGDTFProtocolArtNetDMXMap(const TSharedRef<ArtNet::FDMXGDTFProtocolArtNet>& InProtocolArtNet)
 			: OuterProtocolArtNet(InProtocolArtNet)
 		{}
 	}
diff --git a/Engine/Plugins/VirtualProduction/DMX/DMXGDTF/Source/DMXGDTF/Private/GDTF/Protocols/DMXGDTFProtocols.cpp b/Engine/Plugins/VirtualProduction/DMX/DMXGDTF/Source/DMXGDTF/Private/GDTF/Protocols/DMXGDTFProtocols.cpp
--- a/Engine/Plugins/VirtualProduction/DMX/DMXGDTF/Source/DMXGDTF/Private/GDTF/Protocols/DMXGDTFProtocols.cpp
+++ b/Engine/Plugins/VirtualProduction/DMX/DMXGDTF/Source/DMXGDTF/Private/GDTF/Protocols/DMXGDTFProtocols.cpp
@@ -2,8 +2,11 @@
 
 #include "GDTF/Protocols/DMXGDTFProtocols.h"
 
-#include "GDTF/Protocols/RDM/DMXGDTFProtocolRDM.h"
+#include "GDTF/Protocols/ArtNet/DMXGDTFProtocolArtNet.h"
+#include "GDTF/Protocols/RDM/DMXGDTFProtocolFTRDM.h"
+#include "GDTF/Protocols/sACN/DMXGDTFProtocolSACN.h"
 #include "Serialization/DMXGDTFNodeInitializer.h"
+#include "Serialization/DMXGDTFXmlNodeBuilder.h"
 
 namespace UE::DMX::GDTF
 {
@@ -16,6 +19,18 @@ namespace UE::DMX::GDTF
 		using namespace RDM;
 
 		FDMXGDTFNodeInitializer(SharedThis(this), XmlNode)
-			.CreateOptionalChild(TEXT("RDM"), RDM);
+			.CreateOptionalChild(TEXT("FTRDM"), RDM)
+			.CreateOptionalChild(TEXT("ArtNet"), ArtNet)
+			.CreateOptionalChild(TEXT("sACN"), sACN);
+	}
+
+	FXmlNode* FDMXGDTFProtocols::CreateXmlNode(FXmlNode& Parent)
+	{
+		const FDMXGDTFXmlNodeBuilder ChildBuilder = FDMXGDTFXmlNodeBuilder(Parent, *this)
+			.AppendOptionalChild(TEXT("FTRDM"), RDM)
+			.AppendOptionalChild(TEXT("ArtNet"), ArtNet)
+			.AppendOptionalChild(TEXT("sACN"), sACN);
+
+		return ChildBuilder.GetIntermediateXmlNode();
 	}
 }
diff --git a/Engine/Plugins/VirtualProduction/DMX/DMXGDTF/Source/DMXGDTF/Private/GDTF/Protocols/RDM/DMXGDTFDMXPersonality.cpp b/Engine/Plugins/VirtualProduction/DMX/DMXGDTF/Source/DMXGDTF/Private/GDTF/Protocols/RDM/DMXGDTFDMXPersonality.cpp
--- a/Engine/Plugins/VirtualProduction/DMX/DMXGDTF/Source/DMXGDTF/Private/GDTF/Protocols/RDM/DMXGDTFDMXPersonality.cpp
+++ b/Engine/Plugins/VirtualProduction/DMX/DMXGDTF/Source/DMXGDTF/Private/GDTF/Protocols/RDM/DMXGDTFDMXPersonality.cpp
@@ -5,6 +5,7 @@
 #include "Algo/Find.h"
 #include "GDTF/DMXModes/DMXGDTFDMXMode.h"
 #include "Serialization/DMXGDTFNodeInitializer.h"
+#include "Serialization/DMXGDTFXmlNodeBuilder.h"
 
 namespace UE::DMX::GDTF::RDM
 {
@@ -19,6 +20,15 @@ namespace UE::DMX::GDTF::RDM
 			.GetAttribute(TEXT("DMXMode"), DMXMode);
 	}
 
+	FXmlNode* FDMXGDTFDMXPersonality::CreateXmlNode(FXmlNode& Parent)
+	{
+		const FDMXGDTFXmlNodeBuilder ChildBuilder = FDMXGDTFXmlNodeBuilder(Parent, *this)
+			.SetAttribute(TEXT("Value"), Value)
+			.SetAttribute(TEXT("DMXMode"), DMXMode);
+
+		return ChildBuilder.GetIntermediateXmlNode();
+	}
+
 	TSharedPtr<FDMXGDTFDMXMode> FDMXGDTFDMXPersonality::ResolveDMXMode() const
 	{
 		if (TSharedPtr<FDMXGDTFFixtureType> FixtureType = GetFixtureType().Pin())
diff --git a/Engine/Plugins/VirtualProduction/DMX/DMXGDTF/Source/DMXGDTF/Private/GDTF/Protocols/RDM/DMXGDTFSoftwareVersionID.cpp b/Engine/Plugins/VirtualProduction/DMX/DMXGDTF/Source/DMXGDTF/Private/GDTF/Protocols/RDM/DMXGDTFSoftwareVersionID.cpp
--- a/Engine/Plugins/VirtualProduction/DMX/DMXGDTF/Source/DMXGDTF/Private/GDTF/Protocols/RDM/DMXGDTFSoftwareVersionID.cpp
+++ b/Engine/Plugins/VirtualProduction/DMX/DMXGDTF/Source/DMXGDTF/Private/GDTF/Protocols/RDM/DMXGDTFSoftwareVersionID.cpp
@@ -4,10 +4,11 @@
 
 #include "GDTF/Protocols/RDM/DMXGDTFDMXPersonality.h"
 #include "Serialization/DMXGDTFNodeInitializer.h"
+#include "Serialization/DMXGDTFXmlNodeBuilder.h"
 
 namespace UE::DMX::GDTF::RDM
 {
-	FDMXGDTFSoftwareVersionID::FDMXGDTFSoftwareVersionID(const TSharedRef<FDMXGDTFProtocolRDM>& InProtocolRDM)
+	FDMXGDTFSoftwareVersionID::FDMXGDTFSoftwareVersionID(const TSharedRef<FDMXGDTFProtocolFTRDM>& InProtocolRDM)
 		: OuterProtocolRDM(InProtocolRDM)
 	{}
 
@@ -17,4 +18,13 @@ namespace UE::DMX::GDTF::RDM
 			.GetAttribute(TEXT("Value"), Value, &FParse::HexNumber)
 			.CreateChildren(TEXT("DMXPersonality"), DMXPersonalityArray);
 	}
+
+	FXmlNode* FDMXGDTFSoftwareVersionID::CreateXmlNode(FXmlNode& Parent)
+	{
+		const FDMXGDTFXmlNodeBuilder ChildBuilder = FDMXGDTFXmlNodeBuilder(Parent, *this)
+			.SetAttribute(TEXT("Value"), Value)
+			.AppendChildren(TEXT("DMXPersonality"), DMXPersonalityArray);
+
+		return ChildBuilder.GetIntermediateXmlNode();
+	}
 }
diff --git a/Engine/Plugins/VirtualProduction/DMX/DMXGDTF/Source/DMXGDTF/Private/GDTF/Protocols/sACN/DMXGDTFProtocolSACN.cpp b/Engine/Plugins/VirtualProduction/DMX/DMXGDTF/Source/DMXGDTF/Private/GDTF/Protocols/sACN/DMXGDTFProtocolSACN.cpp
--- a/Engine/Plugins/VirtualProduction/DMX/DMXGDTF/Source/DMXGDTF/Private/GDTF/Protocols/sACN/DMXGDTFProtocolSACN.cpp
+++ b/Engine/Plugins/VirtualProduction/DMX/DMXGDTF/Source/DMXGDTF/Private/GDTF/Protocols/sACN/DMXGDTFProtocolSACN.cpp
@@ -4,6 +4,7 @@
 
 #include "GDTF/Protocols/DMXGDTFProtocolDMXMap.h"
 #include "Serialization/DMXGDTFNodeInitializer.h"
+#include "Serialization/DMXGDTFXmlNodeBuilder.h"
 
 namespace UE::DMX::GDTF::SACN
 {
@@ -16,4 +17,12 @@ namespace UE::DMX::GDTF::SACN
 		FDMXGDTFNodeInitializer(SharedThis(this), XmlNode)
 			.CreateChildren(TEXT("Map"), Maps);
 	}
+
+	FXmlNode* FDMXGDTFProtocolSACN::CreateXmlNode(FXmlNode& Parent)
+	{
+		const FDMXGDTFXmlNodeBuilder ChildBuilder = FDMXGDTFXmlNodeBuilder(Parent, *this)
+			.AppendChildren(TEXT("Map"), Maps);
+
+		return ChildBuilder.GetIntermediateXmlNode();
+	}
 }
diff --git a/Engine/Plugins/VirtualProduction/DMX/DMXGDTF/Source/DMXGDTF/Private/GDTF/Revisions/DMXGDTFRevision.cpp b/Engine/Plugins/VirtualProduction/DMX/DMXGDTF/Source/DMXGDTF/Private/GDTF/Revisions/DMXGDTFRevision.cpp
--- a/Engine/Plugins/VirtualProduction/DMX/DMXGDTF/Source/DMXGDTF/Private/GDTF/Revisions/DMXGDTFRevision.cpp
+++ b/Engine/Plugins/VirtualProduction/DMX/DMXGDTF/Source/DMXGDTF/Private/GDTF/Revisions/DMXGDTFRevision.cpp
@@ -3,6 +3,7 @@
 #include "GDTF/Revisions/DMXGDTFRevision.h"
 
 #include "Serialization/DMXGDTFNodeInitializer.h"
+#include "Serialization/DMXGDTFXmlNodeBuilder.h"
 
 namespace UE::DMX::GDTF
 {
@@ -15,20 +16,56 @@ namespace UE::DMX::GDTF
 		FDMXGDTFNodeInitializer(SharedThis(this), XmlNode)
 			.GetAttribute(TEXT("Text"), Text)
 			.GetAttribute(TEXT("Date"), Date, this, &FDMXGDTFRevision::ParseDateTime)
-			.GetAttribute(TEXT("User"), User)
+			.GetAttribute(TEXT("UserID"), UserID)
 			.GetAttribute(TEXT("ModifiedBy"), ModifiedBy);
 	}
 
+	FXmlNode* FDMXGDTFRevision::CreateXmlNode(FXmlNode& Parent)
+	{
+		const FDMXGDTFXmlNodeBuilder ChildBuilder = FDMXGDTFXmlNodeBuilder(Parent, *this)
+			.SetAttribute(TEXT("Text"), Text)
+			.SetAttribute(TEXT("Date"), Date)
+			.SetAttribute(TEXT("UserID"), UserID)
+			.SetAttribute(TEXT("ModifiedBy"), ModifiedBy);
+
+		return ChildBuilder.GetIntermediateXmlNode();
+	}
+
 	FDateTime FDMXGDTFRevision::ParseDateTime(const FString& GDTFString) const
 	{
-		FDateTime DateTime;
-		if (FDateTime::Parse(GDTFString, DateTime))
+		const FRegexPattern DateTimePattern(TEXT("(\\d{4})-(\\d{2})-(\\d{2})T(\\d{2}):(\\d{2}):(\\d{2})"));
+		FRegexMatcher Regex(DateTimePattern, *GDTFString);
+		if (Regex.FindNext())
 		{
-			return DateTime;
-		}
-		else
-		{
-			return FDateTime::MinValue();
+			const FString YearString = Regex.GetCaptureGroup(1);
+			const FString MonthString = Regex.GetCaptureGroup(2);
+			const FString DayString = Regex.GetCaptureGroup(3);
+			const FString HourString = Regex.GetCaptureGroup(4);
+			const FString MinuteString = Regex.GetCaptureGroup(5);
+			const FString SecondString = Regex.GetCaptureGroup(6);
+
+			if (!YearString.IsEmpty() || !MonthString.IsEmpty() || !DayString.IsEmpty() || !HourString.IsEmpty() || !MinuteString.IsEmpty() || !SecondString.IsEmpty())
+			{
+				int32 Year;
+				int32 Month;
+				int32 Day;
+				int32 Hour;
+				int32 Minute;
+				int32 Second;
+				constexpr int32 Millisecond = 0;
+				if (LexTryParseString(Year, *YearString) &&
+					LexTryParseString(Month, *MonthString) &&
+					LexTryParseString(Day, *DayString) &&
+					LexTryParseString(Hour, *HourString) &&
+					LexTryParseString(Minute, *MinuteString) &&
+					LexTryParseString(Second, *SecondString) &&
+					FDateTime::Validate(Year, Month, Day, Hour, Minute, Second, Millisecond))
+				{
+					return FDateTime(Year, Month, Day, Hour, Minute, Second, Millisecond);
+				}
+			}
 		}
+		
+		return FDateTime::MinValue();
 	}
 }
diff --git a/Engine/Plugins/VirtualProduction/DMX/DMXGDTF/Source/DMXGDTF/Private/GDTF/Wheels/DMXGDTFAnimationSystem.cpp b/Engine/Plugins/VirtualProduction/DMX/DMXGDTF/Source/DMXGDTF/Private/GDTF/Wheels/DMXGDTFAnimationSystem.cpp
--- a/Engine/Plugins/VirtualProduction/DMX/DMXGDTF/Source/DMXGDTF/Private/GDTF/Wheels/DMXGDTFAnimationSystem.cpp
+++ b/Engine/Plugins/VirtualProduction/DMX/DMXGDTF/Source/DMXGDTF/Private/GDTF/Wheels/DMXGDTFAnimationSystem.cpp
@@ -3,6 +3,7 @@
 #include "GDTF/Wheels/DMXGDTFAnimationSystem.h"
 
 #include "Serialization/DMXGDTFNodeInitializer.h"
+#include "Serialization/DMXGDTFXmlNodeBuilder.h"
 
 namespace UE::DMX::GDTF
 {
@@ -18,4 +19,15 @@ namespace UE::DMX::GDTF
 			.GetAttribute(TEXT("P3"), P3)
 			.GetAttribute(TEXT("Radius"), Radius);
 	}
+
+	FXmlNode* FDMXGDTFAnimationSystem::CreateXmlNode(FXmlNode& Parent)
+	{
+		const FDMXGDTFXmlNodeBuilder ChildBuilder = FDMXGDTFXmlNodeBuilder(Parent, *this)
+			.SetAttribute(TEXT("P1"), P1)
+			.SetAttribute(TEXT("P2"), P2)
+			.SetAttribute(TEXT("P3"), P3)
+			.SetAttribute(TEXT("Radius"), Radius);
+
+		return ChildBuilder.GetIntermediateXmlNode();
+	}
 }
diff --git a/Engine/Plugins/VirtualProduction/DMX/DMXGDTF/Source/DMXGDTF/Private/GDTF/Wheels/DMXGDTFPrismFacet.cpp b/Engine/Plugins/VirtualProduction/DMX/DMXGDTF/Source/DMXGDTF/Private/GDTF/Wheels/DMXGDTFPrismFacet.cpp
--- a/Engine/Plugins/VirtualProduction/DMX/DMXGDTF/Source/DMXGDTF/Private/GDTF/Wheels/DMXGDTFPrismFacet.cpp
+++ b/Engine/Plugins/VirtualProduction/DMX/DMXGDTF/Source/DMXGDTF/Private/GDTF/Wheels/DMXGDTFPrismFacet.cpp
@@ -3,6 +3,7 @@
 #include "GDTF/Wheels/DMXGDTFPrismFacet.h"
 
 #include "Serialization/DMXGDTFNodeInitializer.h"
+#include "Serialization/DMXGDTFXmlNodeBuilder.h"
 
 namespace UE::DMX::GDTF
 {
@@ -16,4 +17,13 @@ namespace UE::DMX::GDTF
 			.GetAttribute(TEXT("Color"), Color)
 			.GetAttribute(TEXT("Rotation"), Rotation);
 	}
+
+	FXmlNode* FDMXGDTFPrismFacet::CreateXmlNode(FXmlNode& Parent)
+	{
+		const FDMXGDTFXmlNodeBuilder ChildBuilder = FDMXGDTFXmlNodeBuilder(Parent, *this)
+			.SetAttribute(TEXT("Color"), Color)
+			.SetAttribute(TEXT("Rotation"), Rotation, EDMXGDTFMatrixType::Matrix3x3);
+
+		return ChildBuilder.GetIntermediateXmlNode();
+	}
 }
diff --git a/Engine/Plugins/VirtualProduction/DMX/DMXGDTF/Source/DMXGDTF/Private/GDTF/Wheels/DMXGDTFWheel.cpp b/Engine/Plugins/VirtualProduction/DMX/DMXGDTF/Source/DMXGDTF/Private/GDTF/Wheels/DMXGDTFWheel.cpp
--- a/Engine/Plugins/VirtualProduction/DMX/DMXGDTF/Source/DMXGDTF/Private/GDTF/Wheels/DMXGDTFWheel.cpp
+++ b/Engine/Plugins/VirtualProduction/DMX/DMXGDTF/Source/DMXGDTF/Private/GDTF/Wheels/DMXGDTFWheel.cpp
@@ -4,6 +4,7 @@
 
 #include "GDTF/Wheels/DMXGDTFWheelSlot.h"
 #include "Serialization/DMXGDTFNodeInitializer.h"
+#include "Serialization/DMXGDTFXmlNodeBuilder.h"
 
 namespace UE::DMX::GDTF
 {
@@ -17,4 +18,13 @@ namespace UE::DMX::GDTF
 			.GetAttribute(TEXT("Name"), Name)
 			.CreateChildren(TEXT("Slot"), WheelSlotArray);
 	}
+
+	FXmlNode* FDMXGDTFWheel::CreateXmlNode(FXmlNode& Parent)
+	{
+		const FDMXGDTFXmlNodeBuilder ChildBuilder = FDMXGDTFXmlNodeBuilder(Parent, *this)
+			.SetAttribute(TEXT("Name"), Name)
+			.AppendChildren(TEXT("Slot"), WheelSlotArray);
+
+		return ChildBuilder.GetIntermediateXmlNode();
+	}
 }
diff --git a/Engine/Plugins/VirtualProduction/DMX/DMXGDTF/Source/DMXGDTF/Private/GDTF/Wheels/DMXGDTFWheelSlot.cpp b/Engine/Plugins/VirtualProduction/DMX/DMXGDTF/Source/DMXGDTF/Private/GDTF/Wheels/DMXGDTFWheelSlot.cpp
--- a/Engine/Plugins/VirtualProduction/DMX/DMXGDTF/Source/DMXGDTF/Private/GDTF/Wheels/DMXGDTFWheelSlot.cpp
+++ b/Engine/Plugins/VirtualProduction/DMX/DMXGDTF/Source/DMXGDTF/Private/GDTF/Wheels/DMXGDTFWheelSlot.cpp
@@ -5,6 +5,7 @@
 #include "GDTF/Wheels/DMXGDTFAnimationSystem.h"
 #include "GDTF/Wheels/DMXGDTFPrismFacet.h"
 #include "Serialization/DMXGDTFNodeInitializer.h"
+#include "Serialization/DMXGDTFXmlNodeBuilder.h"
 
 namespace UE::DMX::GDTF
 {
@@ -15,9 +16,24 @@ namespace UE::DMX::GDTF
 	void FDMXGDTFWheelSlot::Initialize(const FXmlNode& XmlNode)
 	{
 		FDMXGDTFNodeInitializer(SharedThis(this), XmlNode)
+			.GetAttribute(TEXT("Name"), Name)
 			.GetAttribute(TEXT("Color"), Color)
 			.GetAttribute(TEXT("MediaFileName"), MediaFileName)
 			.CreateChildren(TEXT("Facet"), PrismFacetArray)
 			.CreateOptionalChild(TEXT("AnimationSystem"), AnimationWheel);
 	}
+
+	FXmlNode* FDMXGDTFWheelSlot::CreateXmlNode(FXmlNode& Parent)
+	{
+		const FString DefaultMediaFileName = TEXT("");
+
+		const FDMXGDTFXmlNodeBuilder ChildBuilder = FDMXGDTFXmlNodeBuilder(Parent, *this)
+			.SetAttribute(TEXT("Name"), Name)
+			.SetAttribute(TEXT("Color"), Color)
+			.SetAttribute(TEXT("MediaFileName"), MediaFileName, DefaultMediaFileName)
+			.AppendChildren(TEXT("Facet"), PrismFacetArray)
+			.AppendOptionalChild(TEXT("AnimationSystem"), AnimationWheel);
+
+		return ChildBuilder.GetIntermediateXmlNode();
+	}
 }
diff --git a/Engine/Plugins/VirtualProduction/DMX/DMXGDTF/Source/DMXGDTF/Private/Serialization/DMXGDTFNodeInitializer.h b/Engine/Plugins/VirtualProduction/DMX/DMXGDTF/Source/DMXGDTF/Private/Serialization/DMXGDTFNodeInitializer.h
--- a/Engine/Plugins/VirtualProduction/DMX/DMXGDTF/Source/DMXGDTF/Private/Serialization/DMXGDTFNodeInitializer.h
+++ b/Engine/Plugins/VirtualProduction/DMX/DMXGDTF/Source/DMXGDTF/Private/Serialization/DMXGDTFNodeInitializer.h
@@ -3,15 +3,16 @@
 #pragma once
 
 #include "Algo/Count.h"
+#include "Algo/Find.h"
 #include "Algo/Transform.h"
+#include "DMXGDTFColorCIE1931xyY.h"
 #include "DMXGDTFLog.h"
 #include "GDTF/DMXGDTFFixtureType.h"
 #include "GDTF/DMXModes/DMXGDTFDMXValue.h"
-#include "GDTF/Geometries/DMXGDTFAddress.h"
+#include "GDTF/Geometries/DMXGDTFDMXAddress.h"
 #include "Internationalization/Regex.h"
 #include "Math/Vector2D.h"
 #include "Misc/Char.h"
-#include "Templates/UnrealTemplate.h"
 #include "UObject/Class.h"
 #include "XmlNode.h"
 #include <type_traits>
@@ -25,31 +26,42 @@ namespace UE::DMX::GDTF
 	 */
 	template<typename NodeType>
 	class FDMXGDTFNodeInitializer
-		: FNoncopyable
+		: public FNoncopyable
 	{
 	public:
-		FDMXGDTFNodeInitializer(TSharedRef<NodeType> InUserObject, const FXmlNode& InXmlNode UE_LIFETIMEBOUND)
+		FDMXGDTFNodeInitializer() = delete;
+		FDMXGDTFNodeInitializer(TSharedRef<NodeType> InUserObject, const FXmlNode& InXmlNode)
 			: UserObject(InUserObject)
 			, XmlNode(InXmlNode)
 		{
-			checkf(XmlNode.GetTag() == InUserObject->GetXmlTag(), TEXT("Tag mismatch when trying to initialize GDTF node initializer. Expected tag '%s' but got tag '%s'."), InUserObject->GetXmlTag(), *XmlNode.GetTag());
+			ensureMsgf(XmlNode.GetTag() == InUserObject->GetXmlTag(), TEXT("Tag mismatch when trying to initialize GDTF node initializer. Expected tag '%s' but got tag '%s'."), InUserObject->GetXmlTag(), *XmlNode.GetTag());
 		}
-		
+	
 		/** Initializes an attribute */
 		template <typename AttributeType>
 		const FDMXGDTFNodeInitializer& GetAttribute(const TCHAR* InAttribute, AttributeType& OutValue) const
 		{
-			const FString StringValue = XmlNode.GetAttribute(InAttribute);
+			const FXmlAttribute* AttributePtr = Algo::FindBy(XmlNode.GetAttributes(), InAttribute, &FXmlAttribute::GetTag);
+			if (!AttributePtr)
+			{
+				return *this;
+			}
+			const FString StringValue = AttributePtr->GetValue();
+
 			if constexpr (std::is_same_v<AttributeType, FString> || std::is_same_v<AttributeType, FName>)
 			{
 				OutValue = *StringValue;
 			}
 			else if constexpr (std::is_arithmetic_v<AttributeType>)
 			{
-				const bool bSuccess = LexTryParseString(OutValue, *StringValue);
-				if (!StringValue.IsEmpty() && !bSuccess)
+				AttributeType ArithmeticValue;
+				if (LexTryParseString(ArithmeticValue, *StringValue))
 				{
-					UE_LOG(LogDMXGDTF, Warning, TEXT("Failed to parse numerical value from XML attribute '%s'. String was '%s'."), InAttribute, *StringValue);
+					OutValue = ArithmeticValue;
+				}
+				else if(!StringValue.IsEmpty())
+				{
+					UE_LOG(LogDMXGDTF, Warning, TEXT("Failed to parse numerical value from XML attribute '%s' in node '%s'. String was '%s'."), InAttribute, UserObject->GetXmlTag(), *StringValue);
 				}
 			}
 			else if constexpr (std::is_enum_v<AttributeType> && std::is_invocable_v<decltype(&FDMXGDTFNodeInitializer::ParseUEnum<AttributeType>), FDMXGDTFNodeInitializer, const FString&, AttributeType&>)
@@ -59,6 +71,10 @@ namespace UE::DMX::GDTF
 				{
 					OutValue = EnumValue;
 				}
+				else
+				{
+					UE_LOG(LogDMXGDTF, Warning, TEXT("Failed to parse enum value from XML attribute '%s' in node '%s'. String was '%s'."), InAttribute, UserObject->GetXmlTag(), *StringValue);
+				}
 			}
 			else if constexpr (std::is_same_v<AttributeType, FVector2D>)
 			{
@@ -68,6 +84,10 @@ namespace UE::DMX::GDTF
 			{
 				OutValue = ParseVector(StringValue);
 			}
+			else if constexpr (std::is_same_v<AttributeType, FDMXGDTFColorCIE1931xyY>)
+			{
+				OutValue = ParseColorCIE(StringValue);
+			}
 			else if constexpr (std::is_same_v<AttributeType, FTransform>)
 			{
 				OutValue = ParseTransform(StringValue);
@@ -76,7 +96,7 @@ namespace UE::DMX::GDTF
 			{
 				OutValue = FGuid(StringValue);
 			}
-			else if constexpr (std::is_same_v<AttributeType, FDMXGDTFAddress>)
+			else if constexpr (std::is_same_v<AttributeType, FDMXGDTFDMXAddress>)
 			{
 				OutValue = ParseDMXAddress(StringValue);
 			}
@@ -306,7 +326,6 @@ namespace UE::DMX::GDTF
 				return FVector::ZeroVector;
 			}
 
-			// Could be a color or a vector. Vectors are wrapped with curly brackets, otherwise they're formated the same.
 			const FString VectorString = GDTFString.Replace(TEXT("{"), TEXT("")).Replace(TEXT("}"), TEXT(""));
 
 			TArray<FString> Substrings;
@@ -321,7 +340,7 @@ namespace UE::DMX::GDTF
 			FVector Result;
 			bool bSuccess = LexTryParseString(Result.X, *Substrings[0]);
 			bSuccess |= LexTryParseString(Result.Y, *Substrings[1]);
-			bSuccess |= LexTryParseString(Result.Z, *Substrings[1]);
+			bSuccess |= LexTryParseString(Result.Z, *Substrings[2]);
 			if (!bSuccess)
 			{
 				UE_LOG(LogDMXGDTF, Warning, TEXT("Cannot parse GDTF vector 3D. Failed to parse %s."), *GDTFString);
@@ -331,117 +350,82 @@ namespace UE::DMX::GDTF
 			return Result;
 		}
 
-		/** Parses a GDTF string as transform */
-		FTransform ParseTransform(const FString& GDTFString) const
+		/** Parses a GDTF string as color CIE */
+		FDMXGDTFColorCIE1931xyY ParseColorCIE(const FString& GDTFString) const
 		{
 			if (GDTFString.IsEmpty())
 			{
-				return FTransform::Identity;
+				return FDMXGDTFColorCIE1931xyY();
 			}
+			
+			TArray<FString> Substrings;
+			GDTFString.ParseIntoArray(Substrings, TEXT(","));
 
-			// Note 3x3 and 4x4 matrices are supported
-
-			// Parse rows
-			TArray<FString> RowStrings;
-
-			const FRegexPattern CurlyBracketRegex(TEXT("\\{([^}]+)\\}"), ERegexPatternFlags::CaseInsensitive);
-			FRegexMatcher CurlyBracketMatcher(CurlyBracketRegex, GDTFString);
-			while (CurlyBracketMatcher.FindNext())
+			if (Substrings.Num() != 3)
 			{
-				RowStrings.Add(CurlyBracketMatcher.GetCaptureGroup(1));
+				UE_LOG(LogDMXGDTF, Warning, TEXT("Cannot parse GDTF color. Expected none or three components, but got %i."), Substrings.Num());
+				return FDMXGDTFColorCIE1931xyY();
 			}
 
-			if (RowStrings.Num() < 3)
+			FDMXGDTFColorCIE1931xyY Result;
+			bool bSuccess = LexTryParseString(Result.X, *Substrings[0]);
+			bSuccess |= LexTryParseString(Result.Y, *Substrings[1]);
+			bSuccess |= LexTryParseString(Result.YY, *Substrings[2]);
+			if (!bSuccess)
 			{
-				UE_LOG(LogDMXGDTF, Warning, TEXT("Cannot parse 3x3 or 4x4 matrix from GDTF string %s."), *GDTFString);
-				return FTransform::Identity;
+				UE_LOG(LogDMXGDTF, Warning, TEXT("Cannot parse GDTF color. Failed to parse %s."), *GDTFString);
+				return FDMXGDTFColorCIE1931xyY();
 			}
 
-			// Parse each component
-			auto GetComponentsFromStringLambda([](const FString& String)
-				{
-					TArray<FString> Substrings;
-					String.ParseIntoArray(Substrings, TEXT(","));
-
-					TArray<float> Result;
-					Result.AddZeroed(Substrings.Num());
-					for (int32 StringIndex = 0; StringIndex < Substrings.Num(); StringIndex++)
-					{
-						float Value = 0.f;
-						if (LexTryParseString(Result[StringIndex], *Substrings[StringIndex]))
-						{
-							// Convert from mm to cm
-							Result[StringIndex] = Value / 10.f;
-						}
-					}
-					return Result;
-				});
-			const TArray<float> U = GetComponentsFromStringLambda(RowStrings[0]);
-			const TArray<float> V = GetComponentsFromStringLambda(RowStrings[1]);
-			const TArray<float> W = GetComponentsFromStringLambda(RowStrings[2]);
-			const TArray<float> O = RowStrings.IsValidIndex(3) ? GetComponentsFromStringLambda(RowStrings[3]) : TArray<float>({ 0.f, 0.f, 0.f });
+			return Result;
+		}
 
-			// Only 3x3 or 4x4 is supported
-			const uint32 NumRows = RowStrings.Num();
-			const bool bValid = (U.Num() == NumRows && V.Num() == NumRows && W.Num() == NumRows && O.Num() == NumRows);
-			if (!bValid)
+		/** Parses a GDTF string as transform */
+		FTransform ParseTransform(const FString& GDTFString) const
+		{
+			FMatrix GDTFMatrix;
+			if (!ParseGDTFMatrix(GDTFString, GDTFMatrix))
 			{
-				UE_LOG(LogDMXGDTF, Warning, TEXT("Cannot parse 3x3 or 4x4 matrix from GDTF string %s."), *GDTFString);
 				return FTransform::Identity;
 			}
+		
+			// To Column major order
+			const FMatrix ColumnMajorGDTFMatrix = GDTFMatrix.GetTransposed();
 
-			// Create the matrix and resulting transform
-			FMatrix Matrix;
-			Matrix.M[0][0] = U[0];
-			Matrix.M[0][1] = U[1];
-			Matrix.M[0][2] = U[2];
-			Matrix.M[0][3] = U.IsValidIndex(3) ? U[3] : 0.f;
-
-			Matrix.M[1][0] = V[0];
-			Matrix.M[1][1] = V[1];
-			Matrix.M[1][2] = V[2];
-			Matrix.M[1][3] = V.IsValidIndex(3) ? U[3] : 0.f;
-
-			Matrix.M[2][0] = W[0];
-			Matrix.M[2][1] = W[0];
-			Matrix.M[2][2] = W[0];
-			Matrix.M[2][3] = W.IsValidIndex(3) ? W[3] : 0.f;
-
-			Matrix.M[3][0] = O[0];
-			Matrix.M[3][1] = -O[1]; // From GDTF's right hand to UE's left hand coordinate system
-			Matrix.M[3][2] = O[2];
-			Matrix.M[3][3] = O.IsValidIndex(3) ? O[3] : 1.f;
-
-			FTransform Transform = FTransform(Matrix);
-
-			// Prevent the convertion to cm from affecting scale
-			Transform.SetScale3D(Transform.GetScale3D() * 10.0);
+			// To Unreal's coordinate system
+			const FMatrix GDTFToUnrealMatrix = FMatrix(
+				FPlane(1.0, 0.0, 0.0, 0.0),
+				FPlane(0.0, 0.0, 1.0, 0.0),
+				FPlane(0.0, -1.0, 0.0, 0.0),
+				FPlane(0.0, 0.0, 0.0, 1.0)
+			);
 
-			// GDTFs are facing down on the Z-Axis, but UE Actors are facing up 
-			const FQuat InvertUpRotationQuaternion = FQuat(Transform.GetRotation().GetAxisY(), PI);
-			Transform.SetRotation(Transform.GetRotation() * InvertUpRotationQuaternion);
+			const FMatrix UnrealMatrix = GDTFToUnrealMatrix * ColumnMajorGDTFMatrix * GDTFToUnrealMatrix;
+			const FTransform Transform = FTransform(UnrealMatrix);
 
 			return Transform;
 		}
 
-		///** Parses a GDTF string as DMX address */
-		FDMXGDTFAddress ParseDMXAddress(const FString& GDTFString) const
+		/** Parses a GDTF string as DMX address */
+		FDMXGDTFDMXAddress ParseDMXAddress(const FString& GDTFString) const
 		{
 			if (GDTFString.IsEmpty())
 			{
-				return FDMXGDTFAddress();
+				return FDMXGDTFDMXAddress();
 			}
 
+			const FString CleanString = GDTFString.Replace(TEXT("{"), TEXT("")).Replace(TEXT("}"), TEXT(""));
+
 			// Can be an absolute adddress
 			uint64 AbsoluteAddress;
-			if (LexTryParseString(AbsoluteAddress, *GDTFString))
+			if (LexTryParseString(AbsoluteAddress, *CleanString))
 			{
-				return FDMXGDTFAddress();
+				return FDMXGDTFDMXAddress(AbsoluteAddress);
 			}
 
 			// Or in the form of Universe.Channel
 			TArray<FString> Substrings;
-			GDTFString.ParseIntoArray(Substrings, TEXT("."));
+			CleanString.ParseIntoArray(Substrings, TEXT("."));
 
 			if (Substrings.Num() == 2)
 			{
@@ -452,18 +436,74 @@ namespace UE::DMX::GDTF
 				{
 					constexpr uint16 UniverseSize = 512;
 
-					FDMXGDTFAddress Result;
+					FDMXGDTFDMXAddress Result;
 					Result.AbsoluteAddress = Universe * UniverseSize + Channel - 1;
 					return Result;
 				}
 			}
 
 			UE_LOG(LogDMXGDTF, Warning, TEXT("Failed to parse DMX Address. '%s' is not a valid GDTF string"), *GDTFString);
-			return FDMXGDTFAddress();
+			return FDMXGDTFDMXAddress();
+		}
+
+		/** Helper to parse 3x3 and 4x4 matrices */
+		bool ParseGDTFMatrix(const FString& InGDTFString, FMatrix& OutMatrix) const
+		{
+			if (InGDTFString.IsEmpty())
+			{
+				return false;
+			}
+
+			// Parse rows
+			TArray<FString> RowStrings;
+
+			const FRegexPattern CurlyBracketRegex(TEXT("\\{([^}]+)\\}"), ERegexPatternFlags::CaseInsensitive);
+			FRegexMatcher CurlyBracketMatcher(CurlyBracketRegex, InGDTFString);
+			while (CurlyBracketMatcher.FindNext())
+			{
+				RowStrings.Add(CurlyBracketMatcher.GetCaptureGroup(1));
+			}
+
+			// Parse each component
+			auto GetComponentsFromStringLambda([](const FString& String)
+				{
+					TArray<FString> Substrings;
+					String.ParseIntoArray(Substrings, TEXT(","));
+
+					TArray<float> Result;
+					Result.AddZeroed(Substrings.Num());
+					for (int32 StringIndex = 0; StringIndex < Substrings.Num(); StringIndex++)
+					{
+						LexTryParseString(Result[StringIndex], *Substrings[StringIndex]);
+					}
+					return Result;
+				});
+			const TArray<float> U = GetComponentsFromStringLambda(RowStrings[0]);
+			const TArray<float> V = GetComponentsFromStringLambda(RowStrings[1]);
+			const TArray<float> W = GetComponentsFromStringLambda(RowStrings[2]);
+			const TArray<float> O = RowStrings.IsValidIndex(3) ? GetComponentsFromStringLambda(RowStrings[3]) : TArray<float>({ 0.f, 0.f, 0.f });
+
+			// Only 3x3 or 4x4 is supported
+			const uint32 NumRows = RowStrings.Num();
+			const bool bValid = (U.Num() == NumRows && V.Num() == NumRows && W.Num() == NumRows && O.Num() == NumRows);
+			if (!bValid)
+			{
+				UE_LOG(LogDMXGDTF, Warning, TEXT("Cannot parse 3x3 or 4x4 matrix from GDTF string '%s'."), *InGDTFString);
+				return false;
+			}
+
+			OutMatrix = FMatrix(
+				FPlane(U[0], U[1], U[2], U.IsValidIndex(3) ? U[3] : 0.f),
+				FPlane(V[0], V[1], V[2], V.IsValidIndex(3) ? V[3] : 0.f),
+				FPlane(W[0], W[1], W[2], W.IsValidIndex(3) ? W[3] : 0.f),
+				FPlane(O[0], O[1], O[2], O.IsValidIndex(3) ? O[3] : 1.f)
+			);
+
+			return true;
 		}
 
 		/** The object that uses this parser */
-		const TSharedRef<NodeType> UserObject;
+		TSharedRef<NodeType> UserObject;
 
 		/** The XML node that needs to be parsed */
 		const FXmlNode& XmlNode;
diff --git a/Engine/Plugins/VirtualProduction/DMX/DMXGDTF/Source/DMXGDTFTests/Private/DMXGDTFTestsModule.cpp b/Engine/Plugins/VirtualProduction/DMX/DMXGDTF/Source/DMXGDTFTests/Private/DMXGDTFTestsModule.cpp
new file mode 100644
--- /dev/null
+++ b/Engine/Plugins/VirtualProduction/DMX/DMXGDTF/Source/DMXGDTFTests/Private/DMXGDTFTestsModule.cpp
@@ -0,0 +1,5 @@
+// Copyright Epic Games, Inc. All Rights Reserved.
+
+#include "Modules/ModuleManager.h"
+
+IMPLEMENT_MODULE(FDefaultModuleImpl, DMXGDTFTestsModule)
diff --git a/Engine/Plugins/VirtualProduction/DMX/DMXGDTF/Source/DMXZip/Private/DMXZipper.cpp b/Engine/Plugins/VirtualProduction/DMX/DMXGDTF/Source/DMXZip/Private/DMXZipper.cpp
--- a/Engine/Plugins/VirtualProduction/DMX/DMXGDTF/Source/DMXZip/Private/DMXZipper.cpp
+++ b/Engine/Plugins/VirtualProduction/DMX/DMXGDTF/Source/DMXZip/Private/DMXZipper.cpp
@@ -39,6 +39,17 @@ bool FDMXZipper::LoadFromData(const TArray64<uint8>& Data)
 }
 
 bool FDMXZipper::SaveToFile(const FString& Filename)
+{
+	TArray64<uint8> Data;
+	if (GetData(Data))
+	{
+		return FFileHelper::SaveArrayToFile(Data, *Filename);
+	}
+	
+	return false;
+}
+
+bool FDMXZipper::GetData(TArray64<uint8>& OutData)
 {
 	FArrayWriter Writer;
 
@@ -78,7 +89,8 @@ bool FDMXZipper::SaveToFile(const FString& Filename)
 
 	AddEndOfCentralDirectory(Writer);
 
-	return FFileHelper::SaveArrayToFile(Writer, *Filename);
+	OutData = Writer;
+	return true;
 }
 
 TArray<FString> FDMXZipper::GetFiles() const
diff --git a/Engine/Plugins/VirtualProduction/DMX/DMXGDTF/Source/DMXZip/Public/DMXZipper.h b/Engine/Plugins/VirtualProduction/DMX/DMXGDTF/Source/DMXZip/Public/DMXZipper.h
--- a/Engine/Plugins/VirtualProduction/DMX/DMXGDTF/Source/DMXZip/Public/DMXZipper.h
+++ b/Engine/Plugins/VirtualProduction/DMX/DMXGDTF/Source/DMXZip/Public/DMXZipper.h
@@ -50,6 +50,9 @@ public:
 	/** Saves the zip to specified Filename */
 	[[nodiscard]] bool SaveToFile(const FString& Filename);
 
+	/** Gets the zip as data. Returns true on success */
+	[[nodiscard]] bool GetData(TArray64<uint8>& OutData);
+
 	/** Returns the File Names in the Zip. Note, may contain relative paths */
 	TArray<FString> GetFiles() const;
 
