DMX - Fix CIS Compile ClangEditor Win64 implicit copy constructor raises a deprecation warning with deprecated member FDMXGDTFDMXChannel::Default

diff --git a/Engine/Plugins/VirtualProduction/DMX/DMXGDTF/Source/DMXGDTF/Internal/GDTF/DMXModes/DMXGDTFDMXChannel.h b/Engine/Plugins/VirtualProduction/DMX/DMXGDTF/Source/DMXGDTF/Internal/GDTF/DMXModes/DMXGDTFDMXChannel.h
--- a/Engine/Plugins/VirtualProduction/DMX/DMXGDTF/Source/DMXGDTF/Internal/GDTF/DMXModes/DMXGDTFDMXChannel.h
+++ b/Engine/Plugins/VirtualProduction/DMX/DMXGDTF/Source/DMXGDTF/Internal/GDTF/DMXModes/DMXGDTFDMXChannel.h
@@ -22,6 +22,9 @@ namespace UE::DMX::GDTF
 	public:
 		FDMXGDTFDMXChannel(const TSharedRef<FDMXGDTFDMXMode>& InDMXMode);
 
+		/** Explicit copy constructor, required to avoid deprecation warnings on Clang with deprecated member Default. */
+		FDMXGDTFDMXChannel(const FDMXGDTFDMXChannel& Other) = default;
+
 		//~ Begin FDMXGDTFNode interface
 		virtual const TCHAR* GetXmlTag() const override { return TEXT("DMXChannel"); }
 		virtual void Initialize(const FXmlNode& XmlNode) override;
