DMX - Fix CIS Compile NonUnity Win64 implicit copy constructor raises a deprecation warning with deprecated member FDMXGDTFDMXChannel::Default

diff --git a/Engine/Plugins/VirtualProduction/DMX/DMXGDTF/Source/DMXGDTF/Internal/GDTF/DMXModes/DMXGDTFDMXChannel.h b/Engine/Plugins/VirtualProduction/DMX/DMXGDTF/Source/DMXGDTF/Internal/GDTF/DMXModes/DMXGDTFDMXChannel.h
--- a/Engine/Plugins/VirtualProduction/DMX/DMXGDTF/Source/DMXGDTF/Internal/GDTF/DMXModes/DMXGDTFDMXChannel.h
+++ b/Engine/Plugins/VirtualProduction/DMX/DMXGDTF/Source/DMXGDTF/Internal/GDTF/DMXModes/DMXGDTFDMXChannel.h
@@ -22,8 +22,14 @@ namespace UE::DMX::GDTF
 	public:
 		FDMXGDTFDMXChannel(const TSharedRef<FDMXGDTFDMXMode>& InDMXMode);
 
-		/** Explicit copy constructor, required to avoid deprecation warnings on Clang with deprecated member Default. */
-		FDMXGDTFDMXChannel(const FDMXGDTFDMXChannel& Other) = default;
+		/** Explicit constructors, required to avoid deprecation warnings on Clang with deprecated member Default. */
+PRAGMA_DISABLE_DEPRECATION_WARNINGS
+		FDMXGDTFDMXChannel(const FDMXGDTFDMXChannel& Node) = delete;
+		FDMXGDTFDMXChannel& operator=(const FDMXGDTFDMXChannel& Node) = delete;
+		FDMXGDTFDMXChannel(FDMXGDTFDMXChannel&& Node) = delete;
+		FDMXGDTFDMXChannel& operator=(FDMXGDTFDMXChannel&& Node) = delete;
+		~FDMXGDTFDMXChannel() = default;
+PRAGMA_ENABLE_DEPRECATION_WARNINGS
 
 		//~ Begin FDMXGDTFNode interface
 		virtual const TCHAR* GetXmlTag() const override { return TEXT("DMXChannel"); }
