DMX - Fix CIS error pragma once in main file

diff --git a/Engine/Plugins/VirtualProduction/DMX/DMXGDTF/Source/DMXGDTFTests/Private/DMXGDTFTests.cpp b/Engine/Plugins/VirtualProduction/DMX/DMXGDTF/Source/DMXGDTFTests/Private/DMXGDTFTests.cpp
--- a/Engine/Plugins/VirtualProduction/DMX/DMXGDTF/Source/DMXGDTFTests/Private/DMXGDTFTests.cpp
+++ b/Engine/Plugins/VirtualProduction/DMX/DMXGDTF/Source/DMXGDTFTests/Private/DMXGDTFTests.cpp
@@ -1,7 +1,5 @@
 // Copyright Epic Games, Inc. All Rights Reserved.
 
-#pragma once
-
 #include "DMXGDTFTests.h"
 
 #include "Algo/Find.h"
