DMXGDTF: Fix potential divide by zero

#rnx

#Synced-CL 50405662

diff --git a/Engine/Plugins/VirtualProduction/DMX/DMXGDTF/Source/DMXGDTF/Private/GDTF/DMXModes/DMXGDTFDMXValue.cpp b/Engine/Plugins/VirtualProduction/DMX/DMXGDTF/Source/DMXGDTF/Private/GDTF/DMXModes/DMXGDTFDMXValue.cpp
--- a/Engine/Plugins/VirtualProduction/DMX/DMXGDTF/Source/DMXGDTF/Private/GDTF/DMXModes/DMXGDTFDMXValue.cpp
+++ b/Engine/Plugins/VirtualProduction/DMX/DMXGDTF/Source/DMXGDTF/Private/GDTF/DMXModes/DMXGDTFDMXValue.cpp
@@ -59,23 +59,25 @@ namespace UE::DMX::GDTF
 
 	uint32 FDMXGDTFDMXValue::GetChecked(const TSharedRef<FDMXGDTFDMXChannel>& InDMXChannel) const
 	{
-		check(IsSet());
-
-		const uint8 WordSize = InDMXChannel->Offset.Num();
-
-		if (WordSize == NumBytes)
-		{
-			return Value;
-		}
-		else if (bByteMirroring)
-		{
-			return static_cast<uint64>(Value) * GetMax(WordSize) / GetMax(NumBytes);
-		}
-		else
+		if (IsSet())
 		{
-			// Byte shift
-			return Value << NumBytes * 8;
+			const uint8 WordSize = InDMXChannel->Offset.Num();
+
+			if (WordSize == NumBytes)
+			{
+				return Value;
+			}
+			else if (bByteMirroring)
+			{
+				return static_cast<uint64>(Value) * GetMax(WordSize) / GetMax(NumBytes);
+			}
+			else
+			{
+				// Byte shift
+				return Value << NumBytes * 8;
+			}
 		}
+		return 0;
 	}
 
 	FString FDMXGDTFDMXValue::AsString() const
